MediaPlate: Added media plate plugin.

#rnx

diff --git a/Engine/Plugins/Media/MediaPlate/MediaPlate.uplugin b/Engine/Plugins/Media/MediaPlate/MediaPlate.uplugin
new file mode 100644
--- /dev/null
+++ b/Engine/Plugins/Media/MediaPlate/MediaPlate.uplugin
@@ -0,0 +1,30 @@
+{
+	"FileVersion": 3,
+	"Version": 0,
+	"VersionName": "0.1",
+	"FriendlyName": "Media Plate",
+	"Description": "Actor that can play media.",
+	"Category": "Media",
+	"CreatedBy": "Epic Games, Inc.",
+	"CreatedByURL": "http://epicgames.com",
+	"DocsURL": "https://docs.unrealengine.com/en-US/Engine/MediaFramework/Overview",
+	"MarketplaceURL": "",
+	"SupportURL": "",
+	"CanContainContent": false,
+	"EnabledByDefault": false,
+	"IsBetaVersion": true,
+	"IsExperimentalVersion": false,
+	"Installed": false,
+	"Modules": [
+		{
+			"Name": "MediaPlate",
+			"Type": "Runtime",
+			"LoadingPhase": "Default"
+		},
+		{
+			"Name": "MediaPlateEditor",
+			"Type": "Editor",
+			"LoadingPhase": "Default"
+		}
+	]
+}
\ No newline at end of file
diff --git a/Engine/Plugins/Media/MediaPlate/Source/MediaPlate/MediaPlate.Build.cs b/Engine/Plugins/Media/MediaPlate/Source/MediaPlate/MediaPlate.Build.cs
new file mode 100644
--- /dev/null
+++ b/Engine/Plugins/Media/MediaPlate/Source/MediaPlate/MediaPlate.Build.cs
@@ -0,0 +1,27 @@
+// Copyright Epic Games, Inc. All Rights Reserved.
+
+using UnrealBuildTool;
+
+public class MediaPlate : ModuleRules
+{
+	public MediaPlate(ReadOnlyTargetRules Target) : base(Target)
+	{
+		PCHUsage = ModuleRules.PCHUsageMode.UseExplicitOrSharedPCHs;
+		
+		PublicDependencyModuleNames.AddRange(
+			new string[]
+			{
+				"Core",
+			}
+			);
+		
+		PrivateDependencyModuleNames.AddRange(
+			new string[]
+			{
+				"CoreUObject",
+				"Engine",
+				"MediaAssets",
+			}
+			);
+	}
+}
diff --git a/Engine/Plugins/Media/MediaPlate/Source/MediaPlate/Private/MediaPlate.cpp b/Engine/Plugins/Media/MediaPlate/Source/MediaPlate/Private/MediaPlate.cpp
new file mode 100644
--- /dev/null
+++ b/Engine/Plugins/Media/MediaPlate/Source/MediaPlate/Private/MediaPlate.cpp
@@ -0,0 +1,15 @@
+// Copyright Epic Games, Inc. All Rights Reserved.
+
+#include "MediaPlate.h"
+#include "MediaPlateModule.h"
+
+#define LOCTEXT_NAMESPACE "MediaPlate"
+
+AMediaPlate::AMediaPlate(const FObjectInitializer& ObjectInitializer)
+	: Super(ObjectInitializer)
+{
+
+}
+
+
+#undef LOCTEXT_NAMESPACE
diff --git a/Engine/Plugins/Media/MediaPlate/Source/MediaPlate/Private/MediaPlateModule.cpp b/Engine/Plugins/Media/MediaPlate/Source/MediaPlate/Private/MediaPlateModule.cpp
new file mode 100644
--- /dev/null
+++ b/Engine/Plugins/Media/MediaPlate/Source/MediaPlate/Private/MediaPlateModule.cpp
@@ -0,0 +1,15 @@
+// Copyright Epic Games, Inc. All Rights Reserved.
+
+#include "MediaPlateModule.h"
+
+DEFINE_LOG_CATEGORY(LogMediaPlate);
+
+void FMediaPlateModule::StartupModule()
+{
+}
+
+void FMediaPlateModule::ShutdownModule()
+{
+}
+
+IMPLEMENT_MODULE(FMediaPlateModule, MediaPlate)
diff --git a/Engine/Plugins/Media/MediaPlate/Source/MediaPlate/Public/MediaPlate.h b/Engine/Plugins/Media/MediaPlate/Source/MediaPlate/Public/MediaPlate.h
new file mode 100644
--- /dev/null
+++ b/Engine/Plugins/Media/MediaPlate/Source/MediaPlate/Public/MediaPlate.h
@@ -0,0 +1,17 @@
+// Copyright Epic Games, Inc. All Rights Reserved.
+
+#pragma once
+
+#include "CoreMinimal.h"
+#include "Engine/StaticMeshActor.h"
+#include "MediaPlate.generated.h"
+
+/**
+ * MediaPlate is an actor that can play and show media in the world.
+ */
+UCLASS()
+class MEDIAPLATE_API AMediaPlate : public AStaticMeshActor
+{
+	GENERATED_UCLASS_BODY()
+
+};
diff --git a/Engine/Plugins/Media/MediaPlate/Source/MediaPlate/Public/MediaPlateModule.h b/Engine/Plugins/Media/MediaPlate/Source/MediaPlate/Public/MediaPlateModule.h
new file mode 100644
--- /dev/null
+++ b/Engine/Plugins/Media/MediaPlate/Source/MediaPlate/Public/MediaPlateModule.h
@@ -0,0 +1,19 @@
+// Copyright Epic Games, Inc. All Rights Reserved.
+
+#pragma once
+
+#include "CoreMinimal.h"
+#include "Logging/LogMacros.h"
+#include "Modules/ModuleManager.h"
+
+/** Log category for this module. */
+DECLARE_LOG_CATEGORY_EXTERN(LogMediaPlate, Log, All);
+
+class FMediaPlateModule : public IModuleInterface
+{
+public:
+	
+	/** IModuleInterface implementation */
+	virtual void StartupModule() override;
+	virtual void ShutdownModule() override;
+};
diff --git a/Engine/Plugins/Media/MediaPlate/Source/MediaPlateEditor/MediaPlateEditor.Build.cs b/Engine/Plugins/Media/MediaPlate/Source/MediaPlateEditor/MediaPlateEditor.Build.cs
new file mode 100644
--- /dev/null
+++ b/Engine/Plugins/Media/MediaPlate/Source/MediaPlateEditor/MediaPlateEditor.Build.cs
@@ -0,0 +1,29 @@
+// Copyright Epic Games, Inc. All Rights Reserved.
+
+using UnrealBuildTool;
+
+public class MediaPlateEditor : ModuleRules
+{
+	public MediaPlateEditor(ReadOnlyTargetRules Target) : base(Target)
+	{
+		PCHUsage = ModuleRules.PCHUsageMode.UseExplicitOrSharedPCHs;
+		
+		PublicDependencyModuleNames.AddRange(
+			new string[]
+			{
+				"Core",
+			}
+			);
+		
+		PrivateDependencyModuleNames.AddRange(
+			new string[]
+			{
+				"CoreUObject",
+				"Engine",
+				"MediaAssets",
+				"UnrealEd",
+				"MediaPlate",
+			}
+			);
+	}
+}
diff --git a/Engine/Plugins/Media/MediaPlate/Source/MediaPlateEditor/Private/ActorFactoryMediaPlate.cpp b/Engine/Plugins/Media/MediaPlate/Source/MediaPlateEditor/Private/ActorFactoryMediaPlate.cpp
new file mode 100644
--- /dev/null
+++ b/Engine/Plugins/Media/MediaPlate/Source/MediaPlateEditor/Private/ActorFactoryMediaPlate.cpp
@@ -0,0 +1,36 @@
+// Copyright Epic Games, Inc. All Rights Reserved.
+
+#include "ActorFactoryMediaPlate.h"
+#include "MediaSource.h"
+#include "MediaPlate.h"
+
+#define LOCTEXT_NAMESPACE "ActorFactoryMediaPlate"
+
+UActorFactoryMediaPlate::UActorFactoryMediaPlate(const FObjectInitializer& ObjectInitializer)
+	: Super(ObjectInitializer)
+{
+	DisplayName = LOCTEXT("MediaPlateDisplayName", "Media Plate");
+	NewActorClass = AMediaPlate::StaticClass();
+}
+
+bool UActorFactoryMediaPlate::CanCreateActorFrom(const FAssetData& AssetData, FText& OutErrorMsg)
+{
+	if (AssetData.IsValid())
+	{
+		UClass* AssetClass = AssetData.GetClass();
+		if ((AssetClass != nullptr) && (AssetClass->IsChildOf(UMediaSource::StaticClass())))
+		{
+			return true;
+		}
+		else
+		{
+			return false;
+		}
+	}
+	else
+	{
+		return true;
+	}
+}
+
+#undef LOCTEXT_NAMESPACE
diff --git a/Engine/Plugins/Media/MediaPlate/Source/MediaPlateEditor/Private/ActorFactoryMediaPlate.h b/Engine/Plugins/Media/MediaPlate/Source/MediaPlateEditor/Private/ActorFactoryMediaPlate.h
new file mode 100644
--- /dev/null
+++ b/Engine/Plugins/Media/MediaPlate/Source/MediaPlateEditor/Private/ActorFactoryMediaPlate.h
@@ -0,0 +1,18 @@
+// Copyright Epic Games, Inc. All Rights Reserved.
+
+#pragma once
+
+#include "CoreMinimal.h"
+#include "UObject/ObjectMacros.h"
+#include "ActorFactories/ActorFactory.h"
+#include "ActorFactoryMediaPlate.generated.h"
+
+UCLASS()
+class UActorFactoryMediaPlate : public UActorFactory
+{
+	GENERATED_UCLASS_BODY()
+
+	//~ Begin UActorFactory Interface
+	virtual bool     CanCreateActorFrom(const FAssetData& AssetData, FText& OutErrorMsg) override;
+	//~ End UActorFactory Interface
+};
diff --git a/Engine/Plugins/Media/MediaPlate/Source/MediaPlateEditor/Private/MediaPlateEditorModule.cpp b/Engine/Plugins/Media/MediaPlate/Source/MediaPlateEditor/Private/MediaPlateEditorModule.cpp
new file mode 100644
--- /dev/null
+++ b/Engine/Plugins/Media/MediaPlate/Source/MediaPlateEditor/Private/MediaPlateEditorModule.cpp
@@ -0,0 +1,15 @@
+// Copyright Epic Games, Inc. All Rights Reserved.
+
+#include "MediaPlateEditorModule.h"
+
+DEFINE_LOG_CATEGORY(LogMediaPlateEditor);
+
+void FMediaPlateEditorModule::StartupModule()
+{
+}
+
+void FMediaPlateEditorModule::ShutdownModule()
+{
+}
+
+IMPLEMENT_MODULE(FMediaPlateEditorModule, MediaPlateEditor)
diff --git a/Engine/Plugins/Media/MediaPlate/Source/MediaPlateEditor/Public/MediaPlateEditorModule.h b/Engine/Plugins/Media/MediaPlate/Source/MediaPlateEditor/Public/MediaPlateEditorModule.h
new file mode 100644
--- /dev/null
+++ b/Engine/Plugins/Media/MediaPlate/Source/MediaPlateEditor/Public/MediaPlateEditorModule.h
@@ -0,0 +1,19 @@
+// Copyright Epic Games, Inc. All Rights Reserved.
+
+#pragma once
+
+#include "CoreMinimal.h"
+#include "Logging/LogMacros.h"
+#include "Modules/ModuleManager.h"
+
+/** Log category for this module. */
+DECLARE_LOG_CATEGORY_EXTERN(LogMediaPlateEditor, Log, All);
+
+class FMediaPlateEditorModule : public IModuleInterface
+{
+public:
+	
+	/** IModuleInterface implementation */
+	virtual void StartupModule() override;
+	virtual void ShutdownModule() override;
+};
