MediaPlate: Added media player, media texture, media source, and a plane mesh.
#rnx

diff --git a/Engine/Plugins/Media/MediaPlate/Source/MediaPlate/Private/MediaPlate.cpp b/Engine/Plugins/Media/MediaPlate/Source/MediaPlate/Private/MediaPlate.cpp
--- a/Engine/Plugins/Media/MediaPlate/Source/MediaPlate/Private/MediaPlate.cpp
+++ b/Engine/Plugins/Media/MediaPlate/Source/MediaPlate/Private/MediaPlate.cpp
@@ -1,14 +1,46 @@
 // Copyright Epic Games, Inc. All Rights Reserved.
 
 #include "MediaPlate.h"
+
+#include "Components/StaticMeshComponent.h"
+#include "Engine/StaticMesh.h"
+#include "MediaPlayer.h"
 #include "MediaPlateModule.h"
+#include "MediaTexture.h"
+#include "UObject/ConstructorHelpers.h"
 
 #define LOCTEXT_NAMESPACE "MediaPlate"
 
 AMediaPlate::AMediaPlate(const FObjectInitializer& ObjectInitializer)
 	: Super(ObjectInitializer)
 {
+	// Create player.
+	MediaPlayer = NewObject<UMediaPlayer>(GetTransientPackage(), MakeUniqueObjectName(GetTransientPackage(), UMediaPlayer::StaticClass()));
+
+	// Create texture.
+	MediaTexture = NewObject<UMediaTexture>(GetTransientPackage(), NAME_None, RF_Transient | RF_Public);
+	if (MediaTexture != nullptr)
+	{
+		MediaTexture->AutoClear = true;
+		MediaTexture->SetMediaPlayer(MediaPlayer);
+		MediaTexture->UpdateResource();
+	}
+
+	// Hook up mesh.
+	struct FConstructorStatics
+	{
+		ConstructorHelpers::FObjectFinder<UStaticMesh> Plane;
+		FConstructorStatics()
+			: Plane(TEXT("/Engine/BasicShapes/Plane")) {}
+	};
 
+	static FConstructorStatics ConstructorStatics;
+	if (ConstructorStatics.Plane.Object != nullptr)
+	{
+		UStaticMeshComponent* LocalStaticMeshComponent = GetStaticMeshComponent();
+		check(LocalStaticMeshComponent);
+		LocalStaticMeshComponent->SetStaticMesh(ConstructorStatics.Plane.Object);
+	}
 }
 
 
diff --git a/Engine/Plugins/Media/MediaPlate/Source/MediaPlate/Public/MediaPlate.h b/Engine/Plugins/Media/MediaPlate/Source/MediaPlate/Public/MediaPlate.h
--- a/Engine/Plugins/Media/MediaPlate/Source/MediaPlate/Public/MediaPlate.h
+++ b/Engine/Plugins/Media/MediaPlate/Source/MediaPlate/Public/MediaPlate.h
@@ -6,6 +6,10 @@
 #include "Engine/StaticMeshActor.h"
 #include "MediaPlate.generated.h"
 
+class UMediaPlayer;
+class UMediaSource;
+class UMediaTexture;
+
 /**
  * MediaPlate is an actor that can play and show media in the world.
  */
@@ -14,4 +18,17 @@ class MEDIAPLATE_API AMediaPlate : public AStaticMeshActor
 {
 	GENERATED_UCLASS_BODY()
 
+public:
+
+	/** Our player to play the media with. */
+	UPROPERTY(transient, EditAnywhere, Category = Media)
+	UMediaPlayer* MediaPlayer;
+
+	/** Our texture that contains the media. */
+	UPROPERTY(transient, EditAnywhere, Category = Media)
+	UMediaTexture* MediaTexture;
+
+	/** What media to play. */
+	UPROPERTY(transient, EditAnywhere, Category = Media)
+	UMediaSource* MediaSource;
 };
diff --git a/Engine/Plugins/Media/MediaPlate/Source/MediaPlateEditor/Private/ActorFactoryMediaPlate.cpp b/Engine/Plugins/Media/MediaPlate/Source/MediaPlateEditor/Private/ActorFactoryMediaPlate.cpp
--- a/Engine/Plugins/Media/MediaPlate/Source/MediaPlateEditor/Private/ActorFactoryMediaPlate.cpp
+++ b/Engine/Plugins/Media/MediaPlate/Source/MediaPlateEditor/Private/ActorFactoryMediaPlate.cpp
@@ -1,6 +1,7 @@
 // Copyright Epic Games, Inc. All Rights Reserved.
 
 #include "ActorFactoryMediaPlate.h"
+
 #include "MediaSource.h"
 #include "MediaPlate.h"
 
@@ -33,4 +34,31 @@ bool UActorFactoryMediaPlate::CanCreateActorFrom(const FAssetData& AssetData, FT
 	}
 }
 
+void UActorFactoryMediaPlate::PostSpawnActor(UObject* Asset, AActor* NewActor)
+{
+	Super::PostSpawnActor(Asset, NewActor);
+
+	SetUpActor(Asset, NewActor);
+}
+
+void UActorFactoryMediaPlate::PostCreateBlueprint(UObject* Asset, AActor* CDO)
+{
+	SetUpActor(Asset, CDO);
+}
+
+void UActorFactoryMediaPlate::SetUpActor(UObject* Asset, AActor* Actor)
+{
+	if (Actor != nullptr)
+	{
+		AMediaPlate* MediaPlate = CastChecked<AMediaPlate>(Actor);
+
+		// Hook up media source.
+		UMediaSource* MediaSource = Cast<UMediaSource>(Asset);
+		if (MediaSource != nullptr)
+		{
+			MediaPlate->MediaSource = MediaSource;
+		}
+	}
+}
+
 #undef LOCTEXT_NAMESPACE
diff --git a/Engine/Plugins/Media/MediaPlate/Source/MediaPlateEditor/Private/ActorFactoryMediaPlate.h b/Engine/Plugins/Media/MediaPlate/Source/MediaPlateEditor/Private/ActorFactoryMediaPlate.h
--- a/Engine/Plugins/Media/MediaPlate/Source/MediaPlateEditor/Private/ActorFactoryMediaPlate.h
+++ b/Engine/Plugins/Media/MediaPlate/Source/MediaPlateEditor/Private/ActorFactoryMediaPlate.h
@@ -13,6 +13,17 @@ class UActorFactoryMediaPlate : public UActorFactory
 	GENERATED_UCLASS_BODY()
 
 	//~ Begin UActorFactory Interface
-	virtual bool     CanCreateActorFrom(const FAssetData& AssetData, FText& OutErrorMsg) override;
+	virtual bool CanCreateActorFrom(const FAssetData& AssetData, FText& OutErrorMsg) override;
+	virtual void PostSpawnActor(UObject* Asset, AActor* NewActor) override;
+	virtual void PostCreateBlueprint(UObject* Asset, AActor* CDO) override;
 	//~ End UActorFactory Interface
+
+private:
+	/**
+	 * Set up a media plate actor.
+	 * 
+	 * @param Asset		Media source asset.
+	 * @param Actor		Media plate actor to set up.
+	 */
+	void SetUpActor(UObject* Asset, AActor* Actor);
 };
