MediaPlate: Uses a MediaComponent to fix assert when calling SetMediaPlayer in constructor.
#rnx

diff --git a/Engine/Plugins/Media/MediaPlate/Source/MediaPlate/Private/MediaPlate.cpp b/Engine/Plugins/Media/MediaPlate/Source/MediaPlate/Private/MediaPlate.cpp
--- a/Engine/Plugins/Media/MediaPlate/Source/MediaPlate/Private/MediaPlate.cpp
+++ b/Engine/Plugins/Media/MediaPlate/Source/MediaPlate/Private/MediaPlate.cpp
@@ -4,44 +4,34 @@
 
 #include "Components/StaticMeshComponent.h"
 #include "Engine/StaticMesh.h"
+#include "MediaComponent.h"
 #include "MediaPlayer.h"
 #include "MediaPlateModule.h"
 #include "MediaTexture.h"
-#include "UObject/ConstructorHelpers.h"
 
 #define LOCTEXT_NAMESPACE "MediaPlate"
 
+FLazyName AMediaPlate::MediaComponentName(TEXT("MediaComponent0"));
+
 AMediaPlate::AMediaPlate(const FObjectInitializer& ObjectInitializer)
 	: Super(ObjectInitializer)
 {
-	// Create player.
-	MediaPlayer = NewObject<UMediaPlayer>(GetTransientPackage(), MakeUniqueObjectName(GetTransientPackage(), UMediaPlayer::StaticClass()));
+	MediaComponent = CreateDefaultSubobject<UMediaComponent>(MediaComponentName);
+}
 
-	// Create texture.
-	MediaTexture = NewObject<UMediaTexture>(GetTransientPackage(), NAME_None, RF_Transient | RF_Public);
-	if (MediaTexture != nullptr)
-	{
-		MediaTexture->AutoClear = true;
-		MediaTexture->SetMediaPlayer(MediaPlayer);
-		MediaTexture->UpdateResource();
-	}
+void AMediaPlate::PostRegisterAllComponents()
+{
+	Super::PostRegisterAllComponents();
 
-	// Hook up mesh.
-	struct FConstructorStatics
-	{
-		ConstructorHelpers::FObjectFinder<UStaticMesh> Plane;
-		FConstructorStatics()
-			: Plane(TEXT("/Engine/BasicShapes/Plane")) {}
-	};
+	UStaticMeshComponent* LocalStaticMeshComponent = GetStaticMeshComponent();
+	check(LocalStaticMeshComponent);
 
-	static FConstructorStatics ConstructorStatics;
-	if (ConstructorStatics.Plane.Object != nullptr)
+	// Add mesh.
+	UStaticMesh* Mesh = LoadObject<UStaticMesh>(nullptr, TEXT("/Engine/BasicShapes/Plane"));
+	if (Mesh != nullptr)
 	{
-		UStaticMeshComponent* LocalStaticMeshComponent = GetStaticMeshComponent();
-		check(LocalStaticMeshComponent);
-		LocalStaticMeshComponent->SetStaticMesh(ConstructorStatics.Plane.Object);
+		LocalStaticMeshComponent->SetStaticMesh(Mesh);
 	}
 }
 
-
 #undef LOCTEXT_NAMESPACE
diff --git a/Engine/Plugins/Media/MediaPlate/Source/MediaPlate/Public/MediaPlate.h b/Engine/Plugins/Media/MediaPlate/Source/MediaPlate/Public/MediaPlate.h
--- a/Engine/Plugins/Media/MediaPlate/Source/MediaPlate/Public/MediaPlate.h
+++ b/Engine/Plugins/Media/MediaPlate/Source/MediaPlate/Public/MediaPlate.h
@@ -6,6 +6,7 @@
 #include "Engine/StaticMeshActor.h"
 #include "MediaPlate.generated.h"
 
+class UMediaComponent;
 class UMediaPlayer;
 class UMediaSource;
 class UMediaTexture;
@@ -19,16 +20,19 @@ class MEDIAPLATE_API AMediaPlate : public AStaticMeshActor
 	GENERATED_UCLASS_BODY()
 
 public:
+	//~ Begin AActor Interface
+	virtual void PostRegisterAllComponents() override;
+	//~ End AActor Interface
 
-	/** Our player to play the media with. */
-	UPROPERTY(transient, EditAnywhere, Category = Media)
-	UMediaPlayer* MediaPlayer;
-
-	/** Our texture that contains the media. */
-	UPROPERTY(transient, EditAnywhere, Category = Media)
-	UMediaTexture* MediaTexture;
+	/** Holds the media player. */
+	UPROPERTY(Category = MediaPlate, VisibleAnywhere, BlueprintReadOnly)
+	TObjectPtr<UMediaComponent> MediaComponent;
 
 	/** What media to play. */
 	UPROPERTY(transient, EditAnywhere, Category = Media)
 	UMediaSource* MediaSource;
+
+private:
+	/** Name for our media component. */
+	static FLazyName MediaComponentName;
 };
