MediaPlate: Fixed mesh being set up from both serialization and construction.
#rnx

diff --git a/Engine/Plugins/Media/MediaPlate/Source/MediaPlate/Private/MediaPlate.cpp b/Engine/Plugins/Media/MediaPlate/Source/MediaPlate/Private/MediaPlate.cpp
--- a/Engine/Plugins/Media/MediaPlate/Source/MediaPlate/Private/MediaPlate.cpp
+++ b/Engine/Plugins/Media/MediaPlate/Source/MediaPlate/Private/MediaPlate.cpp
@@ -8,6 +8,7 @@
 #include "MediaPlayer.h"
 #include "MediaPlateModule.h"
 #include "MediaTexture.h"
+#include "UObject/ConstructorHelpers.h"
 
 #define LOCTEXT_NAMESPACE "MediaPlate"
 
@@ -17,6 +18,23 @@ AMediaPlate::AMediaPlate(const FObjectInitializer& ObjectInitializer)
 	: Super(ObjectInitializer)
 {
 	MediaComponent = CreateDefaultSubobject<UMediaComponent>(MediaComponentName);
+
+	// Hook up mesh.
+	UStaticMeshComponent* LocalStaticMeshComponent = GetStaticMeshComponent();
+	check(LocalStaticMeshComponent);
+	struct FConstructorStatics
+	{
+		ConstructorHelpers::FObjectFinder<UStaticMesh> Plane;
+		FConstructorStatics()
+			: Plane(TEXT("/Engine/BasicShapes/Plane"))
+		{}
+	};
+
+	static FConstructorStatics ConstructorStatics;
+	if (ConstructorStatics.Plane.Object != nullptr)
+	{
+		LocalStaticMeshComponent->SetStaticMesh(ConstructorStatics.Plane.Object);
+	}
 }
 
 void AMediaPlate::PostRegisterAllComponents()
@@ -26,12 +44,6 @@ void AMediaPlate::PostRegisterAllComponents()
 	UStaticMeshComponent* LocalStaticMeshComponent = GetStaticMeshComponent();
 	check(LocalStaticMeshComponent);
 
-	// Add mesh.
-	UStaticMesh* Mesh = LoadObject<UStaticMesh>(nullptr, TEXT("/Engine/BasicShapes/Plane"));
-	if (Mesh != nullptr)
-	{
-		LocalStaticMeshComponent->SetStaticMesh(Mesh);
-	}
 }
 
 #undef LOCTEXT_NAMESPACE
