MediaPlate: Added rendering material.
#rnx

diff --git a/Engine/Plugins/Media/MediaPlate/MediaPlate.uplugin b/Engine/Plugins/Media/MediaPlate/MediaPlate.uplugin
--- a/Engine/Plugins/Media/MediaPlate/MediaPlate.uplugin
+++ b/Engine/Plugins/Media/MediaPlate/MediaPlate.uplugin
@@ -10,7 +10,7 @@
 	"DocsURL": "https://docs.unrealengine.com/en-US/Engine/MediaFramework/Overview",
 	"MarketplaceURL": "",
 	"SupportURL": "",
-	"CanContainContent": false,
+	"CanContainContent": true,
 	"EnabledByDefault": false,
 	"IsBetaVersion": true,
 	"IsExperimentalVersion": false,
diff --git a/Engine/Plugins/Media/MediaPlate/Source/MediaPlate/Private/MediaPlate.cpp b/Engine/Plugins/Media/MediaPlate/Source/MediaPlate/Private/MediaPlate.cpp
--- a/Engine/Plugins/Media/MediaPlate/Source/MediaPlate/Private/MediaPlate.cpp
+++ b/Engine/Plugins/Media/MediaPlate/Source/MediaPlate/Private/MediaPlate.cpp
@@ -13,6 +13,7 @@
 #define LOCTEXT_NAMESPACE "MediaPlate"
 
 FLazyName AMediaPlate::MediaComponentName(TEXT("MediaComponent0"));
+FLazyName AMediaPlate::MediaTextureName("MediaTexture");
 
 AMediaPlate::AMediaPlate(const FObjectInitializer& ObjectInitializer)
 	: Super(ObjectInitializer)
@@ -44,6 +45,16 @@ void AMediaPlate::PostRegisterAllComponents()
 	UStaticMeshComponent* LocalStaticMeshComponent = GetStaticMeshComponent();
 	check(LocalStaticMeshComponent);
 
+	// Add material.
+	UMaterial* Material = LoadObject<UMaterial>(nullptr, TEXT("/MediaPlate/M_MediaPlate"));
+	if (Material != nullptr)
+	{
+		UMaterialInstanceDynamic* MaterialInstance = LocalStaticMeshComponent->CreateAndSetMaterialInstanceDynamicFromMaterial(0, Material);
+		if (MaterialInstance != nullptr)
+		{
+			MaterialInstance->SetTextureParameterValue(MediaTextureName, MediaComponent->GetMediaTexture());
+		}
+	}
 }
 
 #undef LOCTEXT_NAMESPACE
diff --git a/Engine/Plugins/Media/MediaPlate/Source/MediaPlate/Public/MediaPlate.h b/Engine/Plugins/Media/MediaPlate/Source/MediaPlate/Public/MediaPlate.h
--- a/Engine/Plugins/Media/MediaPlate/Source/MediaPlate/Public/MediaPlate.h
+++ b/Engine/Plugins/Media/MediaPlate/Source/MediaPlate/Public/MediaPlate.h
@@ -35,4 +35,6 @@ public:
 private:
 	/** Name for our media component. */
 	static FLazyName MediaComponentName;
+	/** Name for the media texture parameter in the material. */
+	static FLazyName MediaTextureName;
 };
