MediaPlate: Fixed unity build error.
#rnx

diff --git a/Engine/Plugins/Media/MediaPlate/Source/MediaPlate/Private/MediaPlate.cpp b/Engine/Plugins/Media/MediaPlate/Source/MediaPlate/Private/MediaPlate.cpp
--- a/Engine/Plugins/Media/MediaPlate/Source/MediaPlate/Private/MediaPlate.cpp
+++ b/Engine/Plugins/Media/MediaPlate/Source/MediaPlate/Private/MediaPlate.cpp
@@ -4,6 +4,7 @@
 
 #include "Components/StaticMeshComponent.h"
 #include "Engine/StaticMesh.h"
+#include "Materials/MaterialInstanceDynamic.h"
 #include "MediaComponent.h"
 #include "MediaPlayer.h"
 #include "MediaPlateModule.h"
