MediaPlate: Fixed media source not persisting.
#rnx

diff --git a/Engine/Plugins/Media/MediaPlate/Source/MediaPlate/Public/MediaPlate.h b/Engine/Plugins/Media/MediaPlate/Source/MediaPlate/Public/MediaPlate.h
--- a/Engine/Plugins/Media/MediaPlate/Source/MediaPlate/Public/MediaPlate.h
+++ b/Engine/Plugins/Media/MediaPlate/Source/MediaPlate/Public/MediaPlate.h
@@ -29,8 +29,8 @@ public:
 	TObjectPtr<UMediaComponent> MediaComponent;
 
 	/** What media to play. */
-	UPROPERTY(transient, EditAnywhere, Category = Media)
-	UMediaSource* MediaSource;
+	UPROPERTY(EditAnywhere, Category = MediaPlate)
+	TObjectPtr<UMediaSource> MediaSource;
 
 private:
 	/** Name for our media component. */
