MediaPlate: Media texture now has new style output enabled.
#rnx

diff --git a/Engine/Plugins/Media/MediaPlate/Source/MediaPlate/Private/MediaPlate.cpp b/Engine/Plugins/Media/MediaPlate/Source/MediaPlate/Private/MediaPlate.cpp
--- a/Engine/Plugins/Media/MediaPlate/Source/MediaPlate/Private/MediaPlate.cpp
+++ b/Engine/Plugins/Media/MediaPlate/Source/MediaPlate/Private/MediaPlate.cpp
@@ -19,7 +19,17 @@ FLazyName AMediaPlate::MediaTextureName("MediaTexture");
 AMediaPlate::AMediaPlate(const FObjectInitializer& ObjectInitializer)
 	: Super(ObjectInitializer)
 {
+	// Set up media component.
 	MediaComponent = CreateDefaultSubobject<UMediaComponent>(MediaComponentName);
+	if (MediaComponent != nullptr)
+	{
+		// Set up media texture.
+		UMediaTexture* MediaTexture = MediaComponent->GetMediaTexture();
+		if (MediaTexture != nullptr)
+		{
+			MediaTexture->NewStyleOutput = true;
+		}
+	}
 
 	// Hook up mesh.
 	UStaticMeshComponent* LocalStaticMeshComponent = GetStaticMeshComponent();
