MediaPlate: Added some helper functions to play/stop.
#rnx

diff --git a/Engine/Plugins/Media/MediaPlate/Source/MediaPlate/Private/MediaPlate.cpp b/Engine/Plugins/Media/MediaPlate/Source/MediaPlate/Private/MediaPlate.cpp
--- a/Engine/Plugins/Media/MediaPlate/Source/MediaPlate/Private/MediaPlate.cpp
+++ b/Engine/Plugins/Media/MediaPlate/Source/MediaPlate/Private/MediaPlate.cpp
@@ -68,4 +68,34 @@ void AMediaPlate::PostRegisterAllComponents()
 	}
 }
 
+UMediaPlayer* AMediaPlate::GetMediaPlayer()
+{
+	TObjectPtr<UMediaPlayer> MediaPlayer = nullptr;
+	if (MediaComponent != nullptr)
+	{
+		MediaPlayer = MediaComponent->GetMediaPlayer();
+	}
+
+	return MediaPlayer;
+}
+
+void AMediaPlate::Play()
+{
+	TObjectPtr<UMediaPlayer> MediaPlayer = GetMediaPlayer();
+	if (MediaPlayer != nullptr)
+	{
+		MediaPlayer->PlayOnOpen = true;
+		MediaPlayer->OpenSource(MediaSource);
+	}
+}
+
+void AMediaPlate::Stop()
+{
+	TObjectPtr<UMediaPlayer> MediaPlayer = GetMediaPlayer();
+	if (MediaPlayer != nullptr)
+	{
+		MediaPlayer->Close();
+	}
+}
+
 #undef LOCTEXT_NAMESPACE
diff --git a/Engine/Plugins/Media/MediaPlate/Source/MediaPlate/Public/MediaPlate.h b/Engine/Plugins/Media/MediaPlate/Source/MediaPlate/Public/MediaPlate.h
--- a/Engine/Plugins/Media/MediaPlate/Source/MediaPlate/Public/MediaPlate.h
+++ b/Engine/Plugins/Media/MediaPlate/Source/MediaPlate/Public/MediaPlate.h
@@ -24,6 +24,24 @@ public:
 	virtual void PostRegisterAllComponents() override;
 	//~ End AActor Interface
 
+	/**
+	 * Call this get our media player.
+	 */
+	UFUNCTION(BlueprintCallable, Category = "Media|MediaPlate")
+	UMediaPlayer* GetMediaPlayer();
+
+	/**
+	 * Call this to start playing.
+	 */
+	UFUNCTION(BlueprintCallable, Category = "Media|MediaPlate")
+	void Play();
+
+	/**
+	 * Call this to stop playing.
+	 */
+	UFUNCTION(BlueprintCallable, Category = "Media|MediaPlate")
+	void Stop();
+
 	/** Holds the media player. */
 	UPROPERTY(Category = MediaPlate, VisibleAnywhere, BlueprintReadOnly)
 	TObjectPtr<UMediaComponent> MediaComponent;
