MediaPlate: Added customization for MediaPlate actor.

#rnx

diff --git a/Engine/Plugins/Media/MediaPlate/Source/MediaPlateEditor/MediaPlateEditor.Build.cs b/Engine/Plugins/Media/MediaPlate/Source/MediaPlateEditor/MediaPlateEditor.Build.cs
--- a/Engine/Plugins/Media/MediaPlate/Source/MediaPlateEditor/MediaPlateEditor.Build.cs
+++ b/Engine/Plugins/Media/MediaPlate/Source/MediaPlateEditor/MediaPlateEditor.Build.cs
@@ -19,10 +19,13 @@ public class MediaPlateEditor : ModuleRules
 			new string[]
 			{
 				"CoreUObject",
+				"EditorStyle",
 				"Engine",
 				"MediaAssets",
-				"UnrealEd",
 				"MediaPlate",
+				"Slate",
+				"SlateCore",
+				"UnrealEd",
 			}
 			);
 	}
diff --git a/Engine/Plugins/Media/MediaPlate/Source/MediaPlateEditor/Private/MediaPlateCustomization.cpp b/Engine/Plugins/Media/MediaPlate/Source/MediaPlateEditor/Private/MediaPlateCustomization.cpp
new file mode 100644
--- /dev/null
+++ b/Engine/Plugins/Media/MediaPlate/Source/MediaPlateEditor/Private/MediaPlateCustomization.cpp
@@ -0,0 +1,104 @@
+// Copyright Epic Games, Inc. All Rights Reserved.
+
+#pragma once
+
+#include "MediaPlateCustomization.h"
+
+#include "DetailCategoryBuilder.h"
+#include "DetailLayoutBuilder.h"
+#include "DetailWidgetRow.h"
+#include "EditorStyleSet.h"
+#include "MediaPlate.h"
+#include "Widgets/Input/SButton.h"
+#include "Widgets/Images/SImage.h"
+
+#define LOCTEXT_NAMESPACE "FMediaPlateCustomization"
+
+/* IDetailCustomization interface
+ *****************************************************************************/
+
+void FMediaPlateCustomization::CustomizeDetails(IDetailLayoutBuilder& DetailBuilder)
+{
+	IDetailCategoryBuilder& MediaPlateCategory = DetailBuilder.EditCategory("MediaPlate");
+
+	// Get objects we are editing.
+	TArray<TWeakObjectPtr<UObject>> Objects;
+	DetailBuilder.GetObjectsBeingCustomized(Objects);
+	TSharedPtr<TArray<TWeakObjectPtr<AMediaPlate>>> ActorsListPtr = MakeShared<TArray<TWeakObjectPtr<AMediaPlate>>>();
+	ActorsListPtr->Reserve(Objects.Num());
+	for (TWeakObjectPtr<UObject>& Obj : Objects)
+	{
+		TWeakObjectPtr<AMediaPlate> ActorPtr = Cast<AMediaPlate>(Obj.Get());
+		if (ActorPtr.IsValid())
+		{
+			ActorsListPtr->Add(ActorPtr);
+		}
+	}
+
+	// Add media control buttons.
+	MediaPlateCategory.AddCustomRow(LOCTEXT("MediaPlateControls", "MediaPlate Controls"))
+		.NameContent()
+		[
+			SNullWidget::NullWidget
+		]
+		.ValueContent()
+		.VAlign(VAlign_Center)
+		.MaxDesiredWidth(250)
+		[
+			SNew(SHorizontalBox)
+
+			// Play button.
+			+ SHorizontalBox::Slot()
+				.VAlign(VAlign_Center)
+				.AutoWidth()
+				[
+					SNew(SButton)
+						.VAlign(VAlign_Center)
+						.OnClicked_Lambda([ActorsListPtr]() -> FReply
+						{
+							for (TWeakObjectPtr<AMediaPlate>& ActorPtr : *ActorsListPtr)
+							{
+								AMediaPlate* MediaPlate = ActorPtr.Get();
+								if (MediaPlate != nullptr)
+								{
+									MediaPlate->Play();
+								}
+							}
+							return FReply::Handled();
+						})
+						[
+							SNew(SImage)
+								.ColorAndOpacity(FSlateColor::UseForeground())
+								.Image(FEditorStyle::GetBrush("Icons.Toolbar.Play"))
+						]
+				]
+
+			// Stop button.
+			+ SHorizontalBox::Slot()
+				.VAlign(VAlign_Center)
+				.AutoWidth()
+				[
+					SNew(SButton)
+						.VAlign(VAlign_Center)
+						.OnClicked_Lambda([ActorsListPtr]() -> FReply
+						{
+							for (TWeakObjectPtr<AMediaPlate>& ActorPtr : *ActorsListPtr)
+							{
+								AMediaPlate* MediaPlate = ActorPtr.Get();
+								if (MediaPlate != nullptr)
+								{
+									MediaPlate->Stop();
+								}
+							}
+							return FReply::Handled();
+						})
+						[
+							SNew(SImage)
+								.ColorAndOpacity(FSlateColor::UseForeground())
+								.Image(FEditorStyle::GetBrush("Icons.Toolbar.Stop"))
+						]
+				]
+		];
+
+}
+
diff --git a/Engine/Plugins/Media/MediaPlate/Source/MediaPlateEditor/Private/MediaPlateCustomization.h b/Engine/Plugins/Media/MediaPlate/Source/MediaPlateEditor/Private/MediaPlateCustomization.h
new file mode 100644
--- /dev/null
+++ b/Engine/Plugins/Media/MediaPlate/Source/MediaPlateEditor/Private/MediaPlateCustomization.h
@@ -0,0 +1,32 @@
+// Copyright Epic Games, Inc. All Rights Reserved.
+
+#pragma once
+
+#include "CoreMinimal.h"
+#include "IDetailCustomization.h"
+
+class IDetailLayoutBuilder;
+
+/**
+ * Implements a details view customization for the AMediaPlate class.
+ */
+class FMediaPlateCustomization
+	: public IDetailCustomization
+{
+public:
+
+	//~ IDetailCustomization interface
+
+	virtual void CustomizeDetails(IDetailLayoutBuilder& DetailBuilder) override;
+
+	/**
+	 * Creates an instance of this class.
+	 *
+	 * @return The new instance.
+	 */
+	static TSharedRef<IDetailCustomization> MakeInstance()
+	{
+		return MakeShareable(new FMediaPlateCustomization());
+	}
+};
+
diff --git a/Engine/Plugins/Media/MediaPlate/Source/MediaPlateEditor/Private/MediaPlateEditorModule.cpp b/Engine/Plugins/Media/MediaPlate/Source/MediaPlateEditor/Private/MediaPlateEditorModule.cpp
--- a/Engine/Plugins/Media/MediaPlate/Source/MediaPlateEditor/Private/MediaPlateEditorModule.cpp
+++ b/Engine/Plugins/Media/MediaPlate/Source/MediaPlateEditor/Private/MediaPlateEditorModule.cpp
@@ -2,14 +2,27 @@
 
 #include "MediaPlateEditorModule.h"
 
+#include "MediaPlate.h"
+#include "MediaPlateCustomization.h"
+#include "PropertyEditorModule.h"
+
 DEFINE_LOG_CATEGORY(LogMediaPlateEditor);
 
 void FMediaPlateEditorModule::StartupModule()
 {
+	// Register customizations.
+	MediaPlateName = AMediaPlate::StaticClass()->GetFName();
+
+	FPropertyEditorModule& PropertyModule = FModuleManager::LoadModuleChecked<FPropertyEditorModule>("PropertyEditor");
+	PropertyModule.RegisterCustomClassLayout(MediaPlateName,
+		FOnGetDetailCustomizationInstance::CreateStatic(&FMediaPlateCustomization::MakeInstance));
 }
 
 void FMediaPlateEditorModule::ShutdownModule()
 {
+	// Unregister customizations.
+	FPropertyEditorModule& PropertyModule = FModuleManager::LoadModuleChecked<FPropertyEditorModule>("PropertyEditor");
+	PropertyModule.UnregisterCustomClassLayout(MediaPlateName);
 }
 
 IMPLEMENT_MODULE(FMediaPlateEditorModule, MediaPlateEditor)
diff --git a/Engine/Plugins/Media/MediaPlate/Source/MediaPlateEditor/Public/MediaPlateEditorModule.h b/Engine/Plugins/Media/MediaPlate/Source/MediaPlateEditor/Public/MediaPlateEditorModule.h
--- a/Engine/Plugins/Media/MediaPlate/Source/MediaPlateEditor/Public/MediaPlateEditorModule.h
+++ b/Engine/Plugins/Media/MediaPlate/Source/MediaPlateEditor/Public/MediaPlateEditorModule.h
@@ -16,4 +16,9 @@ public:
 	/** IModuleInterface implementation */
 	virtual void StartupModule() override;
 	virtual void ShutdownModule() override;
+
+private:
+
+	/** Customization name to avoid reusing staticstruct during shutdown. */
+	FName MediaPlateName;
 };
