MediaPlate: Static mesh is now offset from the root so it can be rotated to a wall orientation.
#rnx

diff --git a/Engine/Plugins/Media/MediaPlate/Source/MediaPlate/Private/MediaPlate.cpp b/Engine/Plugins/Media/MediaPlate/Source/MediaPlate/Private/MediaPlate.cpp
--- a/Engine/Plugins/Media/MediaPlate/Source/MediaPlate/Private/MediaPlate.cpp
+++ b/Engine/Plugins/Media/MediaPlate/Source/MediaPlate/Private/MediaPlate.cpp
@@ -19,6 +19,8 @@ FLazyName AMediaPlate::MediaTextureName("MediaTexture");
 AMediaPlate::AMediaPlate(const FObjectInitializer& ObjectInitializer)
 	: Super(ObjectInitializer)
 {
+	RootComponent = CreateDefaultSubobject<USceneComponent>(TEXT("Root"));
+
 	// Set up media component.
 	MediaComponent = CreateDefaultSubobject<UMediaComponent>(MediaComponentName);
 	if (MediaComponent != nullptr)
@@ -31,9 +33,12 @@ AMediaPlate::AMediaPlate(const FObjectInitializer& ObjectInitializer)
 		}
 	}
 
+	// Set up static mesh component.
+	StaticMeshComponent = CreateDefaultSubobject<UStaticMeshComponent>(TEXT("StaticMeshComponent"));
+	StaticMeshComponent->SetupAttachment(RootComponent);
+	StaticMeshComponent->SetRelativeRotation(FRotator(0.0f, 90.0f, 90.0f));
+
 	// Hook up mesh.
-	UStaticMeshComponent* LocalStaticMeshComponent = GetStaticMeshComponent();
-	check(LocalStaticMeshComponent);
 	struct FConstructorStatics
 	{
 		ConstructorHelpers::FObjectFinder<UStaticMesh> Plane;
@@ -45,7 +50,7 @@ AMediaPlate::AMediaPlate(const FObjectInitializer& ObjectInitializer)
 	static FConstructorStatics ConstructorStatics;
 	if (ConstructorStatics.Plane.Object != nullptr)
 	{
-		LocalStaticMeshComponent->SetStaticMesh(ConstructorStatics.Plane.Object);
+		StaticMeshComponent->SetStaticMesh(ConstructorStatics.Plane.Object);
 	}
 }
 
@@ -53,17 +58,17 @@ void AMediaPlate::PostRegisterAllComponents()
 {
 	Super::PostRegisterAllComponents();
 
-	UStaticMeshComponent* LocalStaticMeshComponent = GetStaticMeshComponent();
-	check(LocalStaticMeshComponent);
-
-	// Add material.
-	UMaterial* Material = LoadObject<UMaterial>(nullptr, TEXT("/MediaPlate/M_MediaPlate"));
-	if (Material != nullptr)
+	if (StaticMeshComponent != nullptr)
 	{
-		UMaterialInstanceDynamic* MaterialInstance = LocalStaticMeshComponent->CreateAndSetMaterialInstanceDynamicFromMaterial(0, Material);
-		if (MaterialInstance != nullptr)
+		// Add material.
+		UMaterial* Material = LoadObject<UMaterial>(nullptr, TEXT("/MediaPlate/M_MediaPlate"));
+		if (Material != nullptr)
 		{
-			MaterialInstance->SetTextureParameterValue(MediaTextureName, MediaComponent->GetMediaTexture());
+			UMaterialInstanceDynamic* MaterialInstance = StaticMeshComponent->CreateAndSetMaterialInstanceDynamicFromMaterial(0, Material);
+			if (MaterialInstance != nullptr)
+			{
+				MaterialInstance->SetTextureParameterValue(MediaTextureName, MediaComponent->GetMediaTexture());
+			}
 		}
 	}
 }
diff --git a/Engine/Plugins/Media/MediaPlate/Source/MediaPlate/Public/MediaPlate.h b/Engine/Plugins/Media/MediaPlate/Source/MediaPlate/Public/MediaPlate.h
--- a/Engine/Plugins/Media/MediaPlate/Source/MediaPlate/Public/MediaPlate.h
+++ b/Engine/Plugins/Media/MediaPlate/Source/MediaPlate/Public/MediaPlate.h
@@ -3,7 +3,7 @@
 #pragma once
 
 #include "CoreMinimal.h"
-#include "Engine/StaticMeshActor.h"
+#include "GameFramework/Actor.h"
 #include "MediaPlate.generated.h"
 
 class UMediaComponent;
@@ -15,7 +15,7 @@ class UMediaTexture;
  * MediaPlate is an actor that can play and show media in the world.
  */
 UCLASS()
-class MEDIAPLATE_API AMediaPlate : public AStaticMeshActor
+class MEDIAPLATE_API AMediaPlate : public AActor
 {
 	GENERATED_UCLASS_BODY()
 
@@ -46,6 +46,11 @@ public:
 	UPROPERTY(Category = MediaPlate, VisibleAnywhere, BlueprintReadOnly)
 	TObjectPtr<UMediaComponent> MediaComponent;
 
+
+	/** Holds the mesh. */
+	UPROPERTY(Category = MediaPlate, VisibleAnywhere, BlueprintReadOnly)
+	TObjectPtr<UStaticMeshComponent> StaticMeshComponent;
+
 	/** What media to play. */
 	UPROPERTY(EditAnywhere, Category = MediaPlate)
 	TObjectPtr<UMediaSource> MediaSource;
