MediaPlate: Material is only created if it is not already there.
#rnx

diff --git a/Engine/Plugins/Media/MediaPlate/Source/MediaPlate/Private/MediaPlate.cpp b/Engine/Plugins/Media/MediaPlate/Source/MediaPlate/Private/MediaPlate.cpp
--- a/Engine/Plugins/Media/MediaPlate/Source/MediaPlate/Private/MediaPlate.cpp
+++ b/Engine/Plugins/Media/MediaPlate/Source/MediaPlate/Private/MediaPlate.cpp
@@ -60,11 +60,21 @@ void AMediaPlate::PostRegisterAllComponents()
 
 	if (StaticMeshComponent != nullptr)
 	{
-		// Add material.
-		UMaterial* Material = LoadObject<UMaterial>(nullptr, TEXT("/MediaPlate/M_MediaPlate"));
-		if (Material != nullptr)
+		// Do we have a material?
+		if (StaticMeshComponent->GetNumOverrideMaterials() == 0)
 		{
-			UMaterialInstanceDynamic* MaterialInstance = StaticMeshComponent->CreateAndSetMaterialInstanceDynamicFromMaterial(0, Material);
+			// Add material.
+			UMaterial* Material = LoadObject<UMaterial>(nullptr, TEXT("/MediaPlate/M_MediaPlate"));
+			if (Material != nullptr)
+			{
+				UMaterialInstanceDynamic* MaterialInstance = StaticMeshComponent->CreateAndSetMaterialInstanceDynamicFromMaterial(0, Material);
+			}
+		}
+
+		// Set up the material to point to our media texture.
+		if (StaticMeshComponent->GetNumMaterials() > 0)
+		{
+			UMaterialInstanceDynamic* MaterialInstance = Cast< UMaterialInstanceDynamic>(StaticMeshComponent->GetMaterial(0));
 			if (MaterialInstance != nullptr)
 			{
 				MaterialInstance->SetTextureParameterValue(MediaTextureName, MediaComponent->GetMediaTexture());
