MediaPlate: Added autoplay.
#rnx

diff --git a/Engine/Plugins/Media/MediaPlate/Source/MediaPlate/Private/MediaPlate.cpp b/Engine/Plugins/Media/MediaPlate/Source/MediaPlate/Private/MediaPlate.cpp
--- a/Engine/Plugins/Media/MediaPlate/Source/MediaPlate/Private/MediaPlate.cpp
+++ b/Engine/Plugins/Media/MediaPlate/Source/MediaPlate/Private/MediaPlate.cpp
@@ -83,6 +83,16 @@ void AMediaPlate::PostRegisterAllComponents()
 	}
 }
 
+void AMediaPlate::BeginPlay()
+{
+	Super::BeginPlay();
+
+	if (bAutoPlay)
+	{
+		Play();
+	}
+}
+
 UMediaPlayer* AMediaPlate::GetMediaPlayer()
 {
 	TObjectPtr<UMediaPlayer> MediaPlayer = nullptr;
diff --git a/Engine/Plugins/Media/MediaPlate/Source/MediaPlate/Public/MediaPlate.h b/Engine/Plugins/Media/MediaPlate/Source/MediaPlate/Public/MediaPlate.h
--- a/Engine/Plugins/Media/MediaPlate/Source/MediaPlate/Public/MediaPlate.h
+++ b/Engine/Plugins/Media/MediaPlate/Source/MediaPlate/Public/MediaPlate.h
@@ -22,6 +22,7 @@ class MEDIAPLATE_API AMediaPlate : public AActor
 public:
 	//~ Begin AActor Interface
 	virtual void PostRegisterAllComponents() override;
+	virtual void BeginPlay() override;
 	//~ End AActor Interface
 
 	/**
@@ -42,6 +43,10 @@ public:
 	UFUNCTION(BlueprintCallable, Category = "Media|MediaPlate")
 	void Stop();
 
+	/** If set then start playing right away. */
+	UPROPERTY(EditAnywhere, Category = "MediaPlate")
+	bool bAutoPlay;
+
 	/** Holds the media player. */
 	UPROPERTY(Category = MediaPlate, VisibleAnywhere, BlueprintReadOnly)
 	TObjectPtr<UMediaComponent> MediaComponent;
