MediaPlate: Editor viewport is invalidated when a plate is playing so that the movie can be seen.

#rnx

diff --git a/Engine/Plugins/Media/MediaPlate/Source/MediaPlateEditor/Private/MediaPlateCustomization.cpp b/Engine/Plugins/Media/MediaPlate/Source/MediaPlateEditor/Private/MediaPlateCustomization.cpp
--- a/Engine/Plugins/Media/MediaPlate/Source/MediaPlateEditor/Private/MediaPlateCustomization.cpp
+++ b/Engine/Plugins/Media/MediaPlate/Source/MediaPlateEditor/Private/MediaPlateCustomization.cpp
@@ -7,6 +7,7 @@
 #include "DetailWidgetRow.h"
 #include "EditorStyleSet.h"
 #include "MediaPlate.h"
+#include "MediaPlateEditorModule.h"
 #include "Widgets/Input/SButton.h"
 #include "Widgets/Images/SImage.h"
 
@@ -59,6 +60,14 @@ void FMediaPlateCustomization::CustomizeDetails(IDetailLayoutBuilder& DetailBuil
 								AMediaPlate* MediaPlate = ActorPtr.Get();
 								if (MediaPlate != nullptr)
 								{
+									// Tell the editor module that this media plate is playing.
+									FMediaPlateEditorModule* EditorModule = FModuleManager::LoadModulePtr<FMediaPlateEditorModule>("MediaPlateEditor");
+									if (EditorModule != nullptr)
+									{
+										EditorModule->MediaPlateStartedPlayback(MediaPlate);
+									}
+
+									// Play the media.
 									MediaPlate->Play();
 								}
 							}
diff --git a/Engine/Plugins/Media/MediaPlate/Source/MediaPlateEditor/Private/MediaPlateEditorModule.cpp b/Engine/Plugins/Media/MediaPlate/Source/MediaPlateEditor/Private/MediaPlateEditorModule.cpp
--- a/Engine/Plugins/Media/MediaPlate/Source/MediaPlateEditor/Private/MediaPlateEditorModule.cpp
+++ b/Engine/Plugins/Media/MediaPlate/Source/MediaPlateEditor/Private/MediaPlateEditorModule.cpp
@@ -2,8 +2,11 @@
 
 #include "MediaPlateEditorModule.h"
 
+#include "Editor.h"
+#include "LevelEditorViewport.h"
 #include "MediaPlate.h"
 #include "MediaPlateCustomization.h"
+#include "MediaPlayer.h"
 #include "PropertyEditorModule.h"
 
 DEFINE_LOG_CATEGORY(LogMediaPlateEditor);
@@ -25,4 +28,57 @@ void FMediaPlateEditorModule::ShutdownModule()
 	PropertyModule.UnregisterCustomClassLayout(MediaPlateName);
 }
 
+void FMediaPlateEditorModule::Tick(float DeltaTime)
+{
+	bool bIsMediaPlatePlaying = false;
+
+	// Loop through all our plates.
+	for (int Index = 0; Index < ActiveMediaPlates.Num(); ++Index)
+	{
+		// Get the player.
+		TWeakObjectPtr<AMediaPlate>& PlatePtr = ActiveMediaPlates[Index];
+		TObjectPtr<AMediaPlate> MediaPlate = PlatePtr.Get();
+		bool bIsMediaPlateToBeRemoved = true;
+		if (MediaPlate != nullptr)
+		{
+			TObjectPtr<UMediaPlayer> MediaPlayer = MediaPlate->GetMediaPlayer();
+			if (MediaPlayer != nullptr)
+			{
+				bIsMediaPlateToBeRemoved = false;
+				// Are we playing something?
+				if (MediaPlayer->IsPlaying())
+				{
+					bIsMediaPlatePlaying = true;
+					break;
+				}
+			}
+		}
+
+		// Is this player gone?
+		if (bIsMediaPlateToBeRemoved)
+		{
+			ActiveMediaPlates.RemoveAtSwap(Index);
+			--Index;
+		}
+	}
+
+	// Is anything playing?
+	if (bIsMediaPlatePlaying)
+	{
+		// Yes. Invalidate the viewport so we can see it.
+		if (GCurrentLevelEditingViewportClient != nullptr)
+		{
+			GCurrentLevelEditingViewportClient->Invalidate();
+		}
+	}
+}
+
+void FMediaPlateEditorModule::MediaPlateStartedPlayback(TObjectPtr<AMediaPlate> MediaPlate)
+{
+	if (MediaPlate != nullptr)
+	{
+		ActiveMediaPlates.Add(MediaPlate);
+	}
+}
+
 IMPLEMENT_MODULE(FMediaPlateEditorModule, MediaPlateEditor)
diff --git a/Engine/Plugins/Media/MediaPlate/Source/MediaPlateEditor/Public/MediaPlateEditorModule.h b/Engine/Plugins/Media/MediaPlate/Source/MediaPlateEditor/Public/MediaPlateEditorModule.h
--- a/Engine/Plugins/Media/MediaPlate/Source/MediaPlateEditor/Public/MediaPlateEditorModule.h
+++ b/Engine/Plugins/Media/MediaPlate/Source/MediaPlateEditor/Public/MediaPlateEditorModule.h
@@ -5,11 +5,15 @@
 #include "CoreMinimal.h"
 #include "Logging/LogMacros.h"
 #include "Modules/ModuleManager.h"
+#include "TickableEditorObject.h"
+#include "UObject/ObjectPtr.h"
+
+class AMediaPlate;
 
 /** Log category for this module. */
 DECLARE_LOG_CATEGORY_EXTERN(LogMediaPlateEditor, Log, All);
 
-class FMediaPlateEditorModule : public IModuleInterface
+class FMediaPlateEditorModule : public IModuleInterface, public FTickableEditorObject
 {
 public:
 	
@@ -17,8 +21,19 @@ public:
 	virtual void StartupModule() override;
 	virtual void ShutdownModule() override;
 
+	/** FTickableEditorObject interface */
+	virtual void Tick(float DeltaTime) override;
+	virtual TStatId GetStatId() const override { RETURN_QUICK_DECLARE_CYCLE_STAT(FMediaPlateEditorModule, STATGROUP_Tickables); }
+	
+	/**
+	 * Call this when a media plate starts playing so we can track it.
+	 */
+	void MediaPlateStartedPlayback(TObjectPtr<AMediaPlate> MediaPlate);
+	
 private:
 
 	/** Customization name to avoid reusing staticstruct during shutdown. */
 	FName MediaPlateName;
+	/** Holds all the media plates that are playing. */
+	TArray<TWeakObjectPtr<AMediaPlate>> ActiveMediaPlates;
 };
