MediaPlate: Uses the media texture tracker.
#rnx

diff --git a/Engine/Plugins/Media/MediaPlate/Source/MediaPlate/Private/MediaPlate.cpp b/Engine/Plugins/Media/MediaPlate/Source/MediaPlate/Private/MediaPlate.cpp
--- a/Engine/Plugins/Media/MediaPlate/Source/MediaPlate/Private/MediaPlate.cpp
+++ b/Engine/Plugins/Media/MediaPlate/Source/MediaPlate/Private/MediaPlate.cpp
@@ -9,6 +9,7 @@
 #include "MediaPlayer.h"
 #include "MediaPlateModule.h"
 #include "MediaTexture.h"
+#include "MediaTextureTracker.h"
 #include "UObject/ConstructorHelpers.h"
 
 #define LOCTEXT_NAMESPACE "MediaPlate"
@@ -81,6 +82,9 @@ void AMediaPlate::PostRegisterAllComponents()
 			}
 		}
 	}
+
+	// Add our media texture to the tracker.
+	RegisterWithMediaTextureTracker();
 }
 
 void AMediaPlate::BeginPlay()
@@ -93,6 +97,14 @@ void AMediaPlate::BeginPlay()
 	}
 }
 
+void AMediaPlate::BeginDestroy()
+{
+	// Remove our media texture.
+	UnregisterWithMediaTextureTracker();
+	
+	Super::BeginDestroy();
+}
+
 UMediaPlayer* AMediaPlate::GetMediaPlayer()
 {
 	TObjectPtr<UMediaPlayer> MediaPlayer = nullptr;
@@ -104,6 +116,17 @@ UMediaPlayer* AMediaPlate::GetMediaPlayer()
 	return MediaPlayer;
 }
 
+UMediaTexture* AMediaPlate::GetMediaTexture()
+{
+	TObjectPtr<UMediaTexture> MediaTexture = nullptr;
+	if (MediaComponent != nullptr)
+	{
+		MediaTexture = MediaComponent->GetMediaTexture();
+	}
+
+	return MediaTexture;
+}
+
 void AMediaPlate::Play()
 {
 	TObjectPtr<UMediaPlayer> MediaPlayer = GetMediaPlayer();
@@ -123,4 +146,31 @@ void AMediaPlate::Stop()
 	}
 }
 
+void AMediaPlate::RegisterWithMediaTextureTracker()
+{
+	// Set up object.
+	MediaTextureTrackerObject = MakeShared<FMediaTextureTrackerObject, ESPMode::ThreadSafe>();
+	MediaTextureTrackerObject->Object = this;
+	MediaTextureTrackerObject->MipMapLODBias = 0.0f;
+
+	// Add our texture.
+	TObjectPtr<UMediaTexture> MediaTexture = GetMediaTexture();
+	if (MediaTexture != nullptr)
+	{
+		FMediaTextureTracker& MediaTextureTracker = FMediaTextureTracker::Get();
+		MediaTextureTracker.RegisterTexture(MediaTextureTrackerObject, MediaTexture);
+	}
+}
+
+void AMediaPlate::UnregisterWithMediaTextureTracker()
+{
+	// Remove out texture.
+	if (MediaTextureTrackerObject != nullptr)
+	{
+		FMediaTextureTracker& MediaTextureTracker = FMediaTextureTracker::Get();
+		TObjectPtr<UMediaTexture> MediaTexture = GetMediaTexture();
+		MediaTextureTracker.UnregisterTexture(MediaTextureTrackerObject, MediaTexture);
+	}
+}
+
 #undef LOCTEXT_NAMESPACE
diff --git a/Engine/Plugins/Media/MediaPlate/Source/MediaPlate/Public/MediaPlate.h b/Engine/Plugins/Media/MediaPlate/Source/MediaPlate/Public/MediaPlate.h
--- a/Engine/Plugins/Media/MediaPlate/Source/MediaPlate/Public/MediaPlate.h
+++ b/Engine/Plugins/Media/MediaPlate/Source/MediaPlate/Public/MediaPlate.h
@@ -10,6 +10,7 @@ class UMediaComponent;
 class UMediaPlayer;
 class UMediaSource;
 class UMediaTexture;
+struct FMediaTextureTrackerObject;
 
 /**
  * MediaPlate is an actor that can play and show media in the world.
@@ -23,6 +24,7 @@ public:
 	//~ Begin AActor Interface
 	virtual void PostRegisterAllComponents() override;
 	virtual void BeginPlay() override;
+	virtual void BeginDestroy() override;
 	//~ End AActor Interface
 
 	/**
@@ -31,6 +33,12 @@ public:
 	UFUNCTION(BlueprintCallable, Category = "Media|MediaPlate")
 	UMediaPlayer* GetMediaPlayer();
 
+	/**
+	 * Call this get our media texture.
+	 */
+	UFUNCTION(BlueprintCallable, Category = "Media|MediaPlate")
+	UMediaTexture* GetMediaTexture();
+
 	/**
 	 * Call this to start playing.
 	 */
@@ -65,4 +73,16 @@ private:
 	static FLazyName MediaComponentName;
 	/** Name for the media texture parameter in the material. */
 	static FLazyName MediaTextureName;
+
+	/** Info representing this object. */
+	TSharedPtr<FMediaTextureTrackerObject, ESPMode::ThreadSafe> MediaTextureTrackerObject;
+
+	/**
+	 * Adds our media texture to the media texture tracker.
+	 */
+	void RegisterWithMediaTextureTracker();
+	/**
+	 * Removes our texture from the media texture tracker.
+	 */
+	void UnregisterWithMediaTextureTracker();
 };
