MediaPlate: Added editor toolkit for MediaPlate.

#rnx

diff --git a/Engine/Plugins/Media/MediaPlate/Source/MediaPlateEditor/MediaPlateEditor.Build.cs b/Engine/Plugins/Media/MediaPlate/Source/MediaPlateEditor/MediaPlateEditor.Build.cs
--- a/Engine/Plugins/Media/MediaPlate/Source/MediaPlateEditor/MediaPlateEditor.Build.cs
+++ b/Engine/Plugins/Media/MediaPlate/Source/MediaPlateEditor/MediaPlateEditor.Build.cs
@@ -21,8 +21,10 @@ public class MediaPlateEditor : ModuleRules
 				"CoreUObject",
 				"EditorStyle",
 				"Engine",
+				"InputCore",
 				"MediaAssets",
 				"MediaPlate",
+				"MediaPlayerEditor",
 				"Slate",
 				"SlateCore",
 				"UnrealEd",
diff --git a/Engine/Plugins/Media/MediaPlate/Source/MediaPlateEditor/Private/AssetTools/MediaPlateActions.cpp b/Engine/Plugins/Media/MediaPlate/Source/MediaPlateEditor/Private/AssetTools/MediaPlateActions.cpp
new file mode 100644
--- /dev/null
+++ b/Engine/Plugins/Media/MediaPlate/Source/MediaPlateEditor/Private/AssetTools/MediaPlateActions.cpp
@@ -0,0 +1,69 @@
+// Copyright Epic Games, Inc. All Rights Reserved.
+
+#include "AssetTools/MediaPlateActions.h"
+#include "ContentBrowserModule.h"
+#include "EditorStyleSet.h"
+#include "Framework/MultiBox/MultiBoxBuilder.h"
+#include "IContentBrowserSingleton.h"
+#include "MediaPlate.h"
+#include "MediaTexture.h"
+#include "Misc/PackageName.h"
+#include "Toolkits/MediaPlateEditorToolkit.h"
+
+#define LOCTEXT_NAMESPACE "AssetTypeActions"
+
+/* FMediaPlateActions constructors
+ *****************************************************************************/
+
+FMediaPlateActions::FMediaPlateActions(const TSharedRef<ISlateStyle>& InStyle)
+	: Style(InStyle)
+{
+}
+
+/* FAssetTypeActions_Base interface
+ *****************************************************************************/
+
+bool FMediaPlateActions::CanFilter()
+{
+	return true;
+}
+
+uint32 FMediaPlateActions::GetCategories()
+{
+	return EAssetTypeCategories::Media;
+}
+
+FText FMediaPlateActions::GetName() const
+{
+	return NSLOCTEXT("AssetTypeActions", "AssetTypeActions_MediaPlate", "Media Plate");
+}
+
+UClass* FMediaPlateActions::GetSupportedClass() const
+{
+	return AMediaPlate::StaticClass();
+}
+
+FColor FMediaPlateActions::GetTypeColor() const
+{
+	return FColor::Red;
+}
+
+void FMediaPlateActions::OpenAssetEditor(const TArray<UObject*>& InObjects, TSharedPtr<IToolkitHost> EditWithinLevelEditor)
+{
+	EToolkitMode::Type Mode = EditWithinLevelEditor.IsValid()
+		? EToolkitMode::WorldCentric
+		: EToolkitMode::Standalone;
+
+	for (auto ObjIt = InObjects.CreateConstIterator(); ObjIt; ++ObjIt)
+	{
+		AMediaPlate* MediaPlate = Cast<AMediaPlate>(*ObjIt);
+
+		if (MediaPlate != nullptr)
+		{
+			TSharedRef<FMediaPlateEditorToolkit> EditorToolkit = MakeShareable(new FMediaPlateEditorToolkit(Style));
+			EditorToolkit->Initialize(MediaPlate, Mode, EditWithinLevelEditor);
+		}
+	}
+}
+
+#undef LOCTEXT_NAMESPACE
diff --git a/Engine/Plugins/Media/MediaPlate/Source/MediaPlateEditor/Private/AssetTools/MediaPlateActions.h b/Engine/Plugins/Media/MediaPlate/Source/MediaPlateEditor/Private/AssetTools/MediaPlateActions.h
new file mode 100644
--- /dev/null
+++ b/Engine/Plugins/Media/MediaPlate/Source/MediaPlateEditor/Private/AssetTools/MediaPlateActions.h
@@ -0,0 +1,40 @@
+// Copyright Epic Games, Inc. All Rights Reserved.
+
+#pragma once
+
+#include "AssetTypeActions_Base.h"
+#include "Styling/ISlateStyle.h"
+#include "Templates/SharedPointer.h"
+
+class AMediaPlate;
+class FMenuBuilder;
+
+/**
+ * Implements an action for AMediaPlate assets.
+ */
+class FMediaPlateActions
+	: public FAssetTypeActions_Base
+{
+public:
+
+	/**
+	 * Creates and initializes a new instance.
+	 *
+	 * @param InStyle The style set to use for asset editor toolkits.
+	 */
+	FMediaPlateActions(const TSharedRef<ISlateStyle>& InStyle);
+
+	//~ FAssetTypeActions_Base interface
+
+	virtual bool CanFilter() override;
+	virtual uint32 GetCategories() override;
+	virtual FText GetName() const override;
+	virtual UClass* GetSupportedClass() const override;
+	virtual FColor GetTypeColor() const override;
+	virtual void OpenAssetEditor(const TArray<UObject*>& InObjects, TSharedPtr<IToolkitHost> EditWithinLevelEditor = TSharedPtr<IToolkitHost>()) override;
+
+private:
+
+	/** Pointer to the style set to use for toolkits. */
+	TSharedRef<ISlateStyle> Style;
+};
diff --git a/Engine/Plugins/Media/MediaPlate/Source/MediaPlateEditor/Private/MediaPlateCustomization.cpp b/Engine/Plugins/Media/MediaPlate/Source/MediaPlateEditor/Private/MediaPlateCustomization.cpp
--- a/Engine/Plugins/Media/MediaPlate/Source/MediaPlateEditor/Private/MediaPlateCustomization.cpp
+++ b/Engine/Plugins/Media/MediaPlate/Source/MediaPlateEditor/Private/MediaPlateCustomization.cpp
@@ -8,6 +8,7 @@
 #include "EditorStyleSet.h"
 #include "MediaPlate.h"
 #include "MediaPlateEditorModule.h"
+#include "Subsystems/AssetEditorSubsystem.h"
 #include "Widgets/Input/SButton.h"
 #include "Widgets/Images/SImage.h"
 
@@ -23,14 +24,13 @@ void FMediaPlateCustomization::CustomizeDetails(IDetailLayoutBuilder& DetailBuil
 	// Get objects we are editing.
 	TArray<TWeakObjectPtr<UObject>> Objects;
 	DetailBuilder.GetObjectsBeingCustomized(Objects);
-	TSharedPtr<TArray<TWeakObjectPtr<AMediaPlate>>> ActorsListPtr = MakeShared<TArray<TWeakObjectPtr<AMediaPlate>>>();
-	ActorsListPtr->Reserve(Objects.Num());
+	ActorsList.Reserve(Objects.Num());
 	for (TWeakObjectPtr<UObject>& Obj : Objects)
 	{
 		TWeakObjectPtr<AMediaPlate> ActorPtr = Cast<AMediaPlate>(Obj.Get());
 		if (ActorPtr.IsValid())
 		{
-			ActorsListPtr->Add(ActorPtr);
+			ActorsList.Add(ActorPtr);
 		}
 	}
 
@@ -53,9 +53,9 @@ void FMediaPlateCustomization::CustomizeDetails(IDetailLayoutBuilder& DetailBuil
 				[
 					SNew(SButton)
 						.VAlign(VAlign_Center)
-						.OnClicked_Lambda([ActorsListPtr]() -> FReply
+						.OnClicked_Lambda([this]() -> FReply
 						{
-							for (TWeakObjectPtr<AMediaPlate>& ActorPtr : *ActorsListPtr)
+							for (TWeakObjectPtr<AMediaPlate>& ActorPtr : ActorsList)
 							{
 								AMediaPlate* MediaPlate = ActorPtr.Get();
 								if (MediaPlate != nullptr)
@@ -87,9 +87,9 @@ void FMediaPlateCustomization::CustomizeDetails(IDetailLayoutBuilder& DetailBuil
 				[
 					SNew(SButton)
 						.VAlign(VAlign_Center)
-						.OnClicked_Lambda([ActorsListPtr]() -> FReply
+						.OnClicked_Lambda([this]() -> FReply
 						{
-							for (TWeakObjectPtr<AMediaPlate>& ActorPtr : *ActorsListPtr)
+							for (TWeakObjectPtr<AMediaPlate>& ActorPtr : ActorsList)
 							{
 								AMediaPlate* MediaPlate = ActorPtr.Get();
 								if (MediaPlate != nullptr)
@@ -107,5 +107,46 @@ void FMediaPlateCustomization::CustomizeDetails(IDetailLayoutBuilder& DetailBuil
 				]
 		];
 
+
+	// Add button to open the media plate editor.
+	MediaPlateCategory.AddCustomRow(LOCTEXT("OpenMediaPlate", "Open Media Plate"))
+		[
+			SNew(SHorizontalBox)
+
+			+ SHorizontalBox::Slot()
+				.FillWidth(1.f)
+				.Padding(0, 5, 10, 5)
+				[
+					SNew(SButton)
+						.ContentPadding(3)
+						.VAlign(VAlign_Center)
+						.HAlign(HAlign_Center)
+						.OnClicked(this, &FMediaPlateCustomization::OnOpenMediaPlate)
+						.Text(LOCTEXT("OpenMediaPlate", "Open Media Plate"))
+				]
+		];
+}
+
+FReply FMediaPlateCustomization::OnOpenMediaPlate()
+{
+	// Get all our objects.
+	TArray<UObject*> AssetArray;
+	for (TWeakObjectPtr<AMediaPlate>& ActorPtr : ActorsList)
+	{
+		AMediaPlate* MediaPlate = ActorPtr.Get();
+		if (MediaPlate != nullptr)
+		{
+			AssetArray.Add(MediaPlate);
+		}
+	}
+
+	// Open the editor.
+	if (AssetArray.Num() > 0)
+	{
+		GEditor->GetEditorSubsystem<UAssetEditorSubsystem>()->OpenEditorForAssets(AssetArray);
+	}
+
+	return FReply::Handled();
 }
 
+#undef LOCTEXT_NAMESPACE
diff --git a/Engine/Plugins/Media/MediaPlate/Source/MediaPlateEditor/Private/MediaPlateCustomization.h b/Engine/Plugins/Media/MediaPlate/Source/MediaPlateEditor/Private/MediaPlateCustomization.h
--- a/Engine/Plugins/Media/MediaPlate/Source/MediaPlateEditor/Private/MediaPlateCustomization.h
+++ b/Engine/Plugins/Media/MediaPlate/Source/MediaPlateEditor/Private/MediaPlateCustomization.h
@@ -3,8 +3,11 @@
 #pragma once
 
 #include "CoreMinimal.h"
+#include "Editor.h"
 #include "IDetailCustomization.h"
+#include "Input/Reply.h"
 
+class AMediaPlate;
 class IDetailLayoutBuilder;
 
 /**
@@ -28,5 +31,16 @@ public:
 	{
 		return MakeShareable(new FMediaPlateCustomization());
 	}
+
+private:
+
+	/** List of the actors we are editing. */
+	TArray<TWeakObjectPtr<AMediaPlate>> ActorsList;
+
+	/**
+	 * Called when the open media plate button is pressed.
+	 */
+	FReply OnOpenMediaPlate();
+
 };
 
diff --git a/Engine/Plugins/Media/MediaPlate/Source/MediaPlateEditor/Private/MediaPlateEditorModule.cpp b/Engine/Plugins/Media/MediaPlate/Source/MediaPlateEditor/Private/MediaPlateEditorModule.cpp
--- a/Engine/Plugins/Media/MediaPlate/Source/MediaPlateEditor/Private/MediaPlateEditorModule.cpp
+++ b/Engine/Plugins/Media/MediaPlate/Source/MediaPlateEditor/Private/MediaPlateEditorModule.cpp
@@ -2,17 +2,26 @@
 
 #include "MediaPlateEditorModule.h"
 
+#include "AssetTools/MediaPlateActions.h"
 #include "Editor.h"
 #include "LevelEditorViewport.h"
 #include "MediaPlate.h"
 #include "MediaPlateCustomization.h"
+#include "MediaPlateEditorStyle.h"
 #include "MediaPlayer.h"
+#include "Models/MediaPlateEditorCommands.h"
 #include "PropertyEditorModule.h"
 
 DEFINE_LOG_CATEGORY(LogMediaPlateEditor);
 
 void FMediaPlateEditorModule::StartupModule()
 {
+	Style = MakeShareable(new FMediaPlateEditorStyle());
+
+	FMediaPlateEditorCommands::Register();
+
+	RegisterAssetTools();
+
 	// Register customizations.
 	MediaPlateName = AMediaPlate::StaticClass()->GetFName();
 
@@ -23,6 +32,8 @@ void FMediaPlateEditorModule::StartupModule()
 
 void FMediaPlateEditorModule::ShutdownModule()
 {
+	UnregisterAssetTools();
+
 	// Unregister customizations.
 	FPropertyEditorModule& PropertyModule = FModuleManager::LoadModuleChecked<FPropertyEditorModule>("PropertyEditor");
 	PropertyModule.UnregisterCustomClassLayout(MediaPlateName);
@@ -81,4 +92,32 @@ void FMediaPlateEditorModule::MediaPlateStartedPlayback(TObjectPtr<AMediaPlate>
 	}
 }
 
+void FMediaPlateEditorModule::RegisterAssetTools()
+{
+	IAssetTools& AssetTools = FModuleManager::LoadModuleChecked<FAssetToolsModule>("AssetTools").Get();
+
+	RegisterAssetTypeAction(AssetTools, MakeShareable(new FMediaPlateActions(Style.ToSharedRef())));
+}
+
+void FMediaPlateEditorModule::RegisterAssetTypeAction(IAssetTools& AssetTools, TSharedRef<IAssetTypeActions> Action)
+{
+	AssetTools.RegisterAssetTypeActions(Action);
+	RegisteredAssetTypeActions.Add(Action);
+}
+
+void FMediaPlateEditorModule::UnregisterAssetTools()
+{
+	FAssetToolsModule* AssetToolsModule = FModuleManager::GetModulePtr<FAssetToolsModule>("AssetTools");
+
+	if (AssetToolsModule != nullptr)
+	{
+		IAssetTools& AssetTools = AssetToolsModule->Get();
+
+		for (TSharedRef<IAssetTypeActions>& Action : RegisteredAssetTypeActions)
+		{
+			AssetTools.UnregisterAssetTypeActions(Action);
+		}
+	}
+}
+
 IMPLEMENT_MODULE(FMediaPlateEditorModule, MediaPlateEditor)
diff --git a/Engine/Plugins/Media/MediaPlate/Source/MediaPlateEditor/Private/MediaPlateEditorStyle.cpp b/Engine/Plugins/Media/MediaPlate/Source/MediaPlateEditor/Private/MediaPlateEditorStyle.cpp
new file mode 100644
--- /dev/null
+++ b/Engine/Plugins/Media/MediaPlate/Source/MediaPlateEditor/Private/MediaPlateEditorStyle.cpp
@@ -0,0 +1,94 @@
+// Copyright Epic Games, Inc. All Rights Reserved.
+
+#include "MediaPlateEditorStyle.h"
+
+#include "Brushes/SlateBorderBrush.h"
+#include "Brushes/SlateBoxBrush.h"
+#include "Brushes/SlateImageBrush.h"
+#include "Fonts/SlateFontInfo.h"
+#include "Math/Vector2D.h"
+#include "Misc/Paths.h"
+#include "Styling/SlateStyleRegistry.h"
+#include "Styling/SlateTypes.h"
+#include "Styling/CoreStyle.h"
+
+#define IMAGE_BRUSH(RelativePath, ...) FSlateImageBrush(RootToContentDir(RelativePath, TEXT(".png")), __VA_ARGS__)
+#define BOX_BRUSH(RelativePath, ...) FSlateBoxBrush(RootToContentDir(RelativePath, TEXT(".png")), __VA_ARGS__)
+#define BORDER_BRUSH(RelativePath, ...) FSlateBorderBrush(RootToContentDir(RelativePath, TEXT(".png")), __VA_ARGS__)
+#define DEFAULT_FONT(...) FCoreStyle::GetDefaultFontStyle(__VA_ARGS__)
+
+/* FMediaPlateEditorStyle structors
+ *****************************************************************************/
+
+FMediaPlateEditorStyle::FMediaPlateEditorStyle()
+	: FSlateStyleSet("MediaPlateEditorStyle")
+{
+	const FVector2D Icon8x8(8.0f, 8.0f);
+	const FVector2D Icon12x12(12.0f, 12.0f);
+	const FVector2D Icon16x16(16.0f, 16.0f);
+	const FVector2D Icon20x20(20.0f, 20.0f);
+	const FVector2D Icon40x40(40.0f, 40.0f);
+
+	SetContentRoot(FPaths::EnginePluginsDir() / TEXT("Media/MediaPlayerEditor/Content"));
+
+	// buttons
+	Set("MediaPlateEditor.SourceButton", new IMAGE_BRUSH("btn_source_12x", Icon12x12));
+	Set("MediaPlateEditor.GoButton", new IMAGE_BRUSH("btn_go_12x", Icon12x12));
+	Set("MediaPlateEditor.ReloadButton", new IMAGE_BRUSH("btn_reload_12x", Icon12x12));
+	Set("MediaPlateEditor.SettingsButton", new IMAGE_BRUSH("btn_settings_16x", Icon12x12));
+
+	// misc
+	Set("MediaPlateEditor.DragDropBorder", new BOX_BRUSH("border_dragdrop", 0.5f));
+	Set("MediaPlateEditor.MediaSourceOpened", new IMAGE_BRUSH("mediasource_opened", Icon8x8));
+
+	// tabs
+	Set("MediaPlateEditor.Tabs.Info", new IMAGE_BRUSH("tab_info_16x", Icon16x16));
+	Set("MediaPlateEditor.Tabs.Media", new IMAGE_BRUSH("tab_media_16x", Icon16x16));
+	Set("MediaPlateEditor.Tabs.Player", new IMAGE_BRUSH("tab_player_16x", Icon16x16));
+	Set("MediaPlateEditor.Tabs.Playlist", new IMAGE_BRUSH("tab_playlist_16x", Icon16x16));
+	Set("MediaPlateEditor.Tabs.Stats", new IMAGE_BRUSH("tab_stats_16x", Icon16x16));
+
+	// toolbar icons
+	Set("MediaPlateEditor.CloseMedia", new IMAGE_BRUSH("icon_eject_40x", Icon40x40));
+	Set("MediaPlateEditor.ForwardMedia", new IMAGE_BRUSH("icon_forward_40x", Icon40x40));
+	Set("MediaPlateEditor.ForwardMedia.Small", new IMAGE_BRUSH("icon_forward_40x", Icon20x20));
+	Set("MediaPlateEditor.NextMedia", new IMAGE_BRUSH("icon_step_40x", Icon40x40));
+	Set("MediaPlateEditor.NextMedia.Small", new IMAGE_BRUSH("icon_step_40x", Icon20x20));
+	Set("MediaPlateEditor.OpenMedia", new IMAGE_BRUSH("tab_media_16x", Icon40x40));
+	Set("MediaPlateEditor.PauseMedia", new IMAGE_BRUSH("icon_pause_40x", Icon40x40));
+	Set("MediaPlateEditor.PauseMedia.Small", new IMAGE_BRUSH("icon_pause_40x", Icon20x20));
+	Set("MediaPlateEditor.PlayMedia", new IMAGE_BRUSH("icon_play_40x", Icon40x40));
+	Set("MediaPlateEditor.PlayMedia.Small", new IMAGE_BRUSH("icon_play_40x", Icon20x20));
+	Set("MediaPlateEditor.PreviousMedia", new IMAGE_BRUSH("icon_step_back_40x", Icon40x40));
+	Set("MediaPlateEditor.PreviousMedia.Small", new IMAGE_BRUSH("icon_step_back_40x", Icon20x20));
+	Set("MediaPlateEditor.ReverseMedia", new IMAGE_BRUSH("icon_reverse_40x", Icon40x40));
+	Set("MediaPlateEditor.ReverseMedia.Small", new IMAGE_BRUSH("icon_reverse_40x", Icon20x20));
+	Set("MediaPlateEditor.RewindMedia", new IMAGE_BRUSH("icon_rewind_40x", Icon40x40));
+	Set("MediaPlateEditor.RewindMedia.Small", new IMAGE_BRUSH("icon_rewind_40x", Icon20x20));
+	Set("MediaPlateEditor.StopMedia", new IMAGE_BRUSH("icon_stop_40x", Icon40x40));
+	Set("MediaPlateEditor.StopMedia.Small", new IMAGE_BRUSH("icon_stop_40x", Icon20x20));
+
+	// scrubber
+	Set("MediaPlayerEditor.Scrubber", FSliderStyle()
+		.SetNormalBarImage(FSlateColorBrush(FColor::White))
+		.SetDisabledBarImage(FSlateColorBrush(FLinearColor::Gray))
+		.SetNormalThumbImage(IMAGE_BRUSH("scrubber", FVector2D(2.0f, 10.0f)))
+		.SetHoveredThumbImage(IMAGE_BRUSH("scrubber", FVector2D(2.0f, 10.0f)))
+		.SetDisabledThumbImage(IMAGE_BRUSH("scrubber", FVector2D(2.0f, 10.0f)))
+		.SetBarThickness(2.0f)
+	);
+
+	Set("MediaPlateEditor.ViewportFont", DEFAULT_FONT("Regular", 18));
+
+	FSlateStyleRegistry::RegisterSlateStyle(*this);
+}
+
+FMediaPlateEditorStyle::~FMediaPlateEditorStyle()
+{
+	FSlateStyleRegistry::UnRegisterSlateStyle(*this);
+}
+
+#undef IMAGE_BRUSH
+#undef BOX_BRUSH
+#undef BORDER_BRUSH
+#undef DEFAULT_FONT
diff --git a/Engine/Plugins/Media/MediaPlate/Source/MediaPlateEditor/Private/MediaPlateEditorStyle.h b/Engine/Plugins/Media/MediaPlate/Source/MediaPlateEditor/Private/MediaPlateEditorStyle.h
new file mode 100644
--- /dev/null
+++ b/Engine/Plugins/Media/MediaPlate/Source/MediaPlateEditor/Private/MediaPlateEditorStyle.h
@@ -0,0 +1,18 @@
+// Copyright Epic Games, Inc. All Rights Reserved.
+
+#pragma once
+
+#include "Styling/SlateStyle.h"
+
+/**
+ * Implements the visual style of the media plate editor UI.
+ */
+class FMediaPlateEditorStyle : public FSlateStyleSet
+{
+public:
+	/** Default constructor. */
+	FMediaPlateEditorStyle();
+
+	 /** Destructor. */
+	~FMediaPlateEditorStyle();
+};
diff --git a/Engine/Plugins/Media/MediaPlate/Source/MediaPlateEditor/Private/Models/MediaPlateEditorCommands.cpp b/Engine/Plugins/Media/MediaPlate/Source/MediaPlateEditor/Private/Models/MediaPlateEditorCommands.cpp
new file mode 100644
--- /dev/null
+++ b/Engine/Plugins/Media/MediaPlate/Source/MediaPlateEditor/Private/Models/MediaPlateEditorCommands.cpp
@@ -0,0 +1,20 @@
+// Copyright Epic Games, Inc. All Rights Reserved.
+
+#include "Models/MediaPlateEditorCommands.h"
+
+#define LOCTEXT_NAMESPACE "FMediaPlateEditorCommands"
+
+void FMediaPlateEditorCommands::RegisterCommands()
+{
+	UI_COMMAND(CloseMedia, "Close", "Close the currently opened media", EUserInterfaceActionType::Button, FInputChord());
+	UI_COMMAND(ForwardMedia, "Forward", "Fast forward media playback", EUserInterfaceActionType::Button, FInputChord());
+	UI_COMMAND(NextMedia, "Next", "Jump to next item in the play list", EUserInterfaceActionType::Button, FInputChord());
+	UI_COMMAND(OpenMedia, "Open", "Open the current media", EUserInterfaceActionType::Button, FInputChord());
+	UI_COMMAND(PauseMedia, "Pause", "Pause media playback", EUserInterfaceActionType::Button, FInputChord());
+	UI_COMMAND(PlayMedia, "Play", "Start media playback", EUserInterfaceActionType::Button, FInputChord());
+	UI_COMMAND(PreviousMedia, "Prev", "Jump to previous item in the play list", EUserInterfaceActionType::Button, FInputChord());
+	UI_COMMAND(ReverseMedia, "Reverse", "Reverse media playback", EUserInterfaceActionType::Button, FInputChord());
+	UI_COMMAND(RewindMedia, "Rewind", "Rewind the media to the beginning", EUserInterfaceActionType::Button, FInputChord());
+}
+
+#undef LOCTEXT_NAMESPACE
diff --git a/Engine/Plugins/Media/MediaPlate/Source/MediaPlateEditor/Private/Models/MediaPlateEditorCommands.h b/Engine/Plugins/Media/MediaPlate/Source/MediaPlateEditor/Private/Models/MediaPlateEditorCommands.h
new file mode 100644
--- /dev/null
+++ b/Engine/Plugins/Media/MediaPlate/Source/MediaPlateEditor/Private/Models/MediaPlateEditorCommands.h
@@ -0,0 +1,49 @@
+// Copyright Epic Games, Inc. All Rights Reserved.
+
+#pragma once
+
+#include "CoreMinimal.h"
+#include "Framework/Commands/Commands.h"
+
+/**
+ * Holds the UI commands for the MediaPlateEditorToolkit widget.
+ */
+class FMediaPlateEditorCommands : public TCommands<FMediaPlateEditorCommands>
+{
+public:
+	/** Default constructor. */
+	FMediaPlateEditorCommands()
+		: TCommands<FMediaPlateEditorCommands>("MediaPlateEditor", NSLOCTEXT("Contexts", "MediaPlateEditor", "MediaPlate Editor"), NAME_None, "MediaPlateEditorStyle")
+	{ }
+
+	//~ TCommands interface
+	virtual void RegisterCommands() override;
+	
+
+	/** Close the currently opened media. */
+	TSharedPtr<FUICommandInfo> CloseMedia;
+
+	/** Fast forwards media playback. */
+	TSharedPtr<FUICommandInfo> ForwardMedia;
+
+	/** Jump to next item in the play list. */
+	TSharedPtr<FUICommandInfo> NextMedia;
+
+	/** Open the current media. */
+	TSharedPtr<FUICommandInfo> OpenMedia;
+
+	/** Pauses media playback. */
+	TSharedPtr<FUICommandInfo> PauseMedia;
+
+	/** Starts media playback. */
+	TSharedPtr<FUICommandInfo> PlayMedia;
+
+	/** Jump to previous item in the play list. */
+	TSharedPtr<FUICommandInfo> PreviousMedia;
+
+	/** Reverses media playback. */
+	TSharedPtr<FUICommandInfo> ReverseMedia;
+
+	/** Rewinds the media to the beginning. */
+	TSharedPtr<FUICommandInfo> RewindMedia;
+};
diff --git a/Engine/Plugins/Media/MediaPlate/Source/MediaPlateEditor/Private/Toolkits/MediaPlateEditorToolkit.cpp b/Engine/Plugins/Media/MediaPlate/Source/MediaPlateEditor/Private/Toolkits/MediaPlateEditorToolkit.cpp
new file mode 100644
--- /dev/null
+++ b/Engine/Plugins/Media/MediaPlate/Source/MediaPlateEditor/Private/Toolkits/MediaPlateEditorToolkit.cpp
@@ -0,0 +1,430 @@
+// Copyright Epic Games, Inc. All Rights Reserved.
+
+#include "Toolkits/MediaPlateEditorToolkit.h"
+
+#include "Editor.h"
+#include "EditorStyleSet.h"
+#include "EditorReimportHandler.h"
+#include "Framework/MultiBox/MultiBoxBuilder.h"
+#include "MediaPlate.h"
+#include "MediaPlayer.h"
+#include "MediaPlaylist.h"
+#include "Models/MediaPlateEditorCommands.h"
+#include "SlateOptMacros.h"
+#include "Widgets/Docking/SDockTab.h"
+#include "Widgets/SMediaPlayerEditorViewer.h"
+
+#define LOCTEXT_NAMESPACE "FMediaPlateEditorToolkit"
+
+namespace MediaPlateEditorToolkit
+{
+	static const FName AppIdentifier("MediaPlateEditorApp");
+	static const FName DetailsTabId("Details");
+	static const FName InfoTabId("Info");
+	static const FName MediaTabId("Media");
+	static const FName PlaylistTabId("Playlist");
+	static const FName StatsTabId("Stats");
+	static const FName ViewerTabId("Viewer");
+}
+
+/* FMediaPlateEditorToolkit structors
+ *****************************************************************************/
+
+FMediaPlateEditorToolkit::FMediaPlateEditorToolkit(const TSharedRef<ISlateStyle>& InStyle)
+	: MediaPlate(nullptr)
+	, Style(InStyle)
+{
+}
+
+FMediaPlateEditorToolkit::~FMediaPlateEditorToolkit()
+{
+	FReimportManager::Instance()->OnPreReimport().RemoveAll(this);
+	FReimportManager::Instance()->OnPostReimport().RemoveAll(this);
+
+	GEditor->UnregisterForUndo(this);
+}
+
+/* FMediaPlateEditorToolkit interface
+ *****************************************************************************/
+
+void FMediaPlateEditorToolkit::Initialize(AMediaPlate* InMediaPlate, const EToolkitMode::Type InMode, const TSharedPtr<IToolkitHost>& InToolkitHost)
+{
+	MediaPlate = InMediaPlate;
+
+	if (MediaPlate == nullptr)
+	{
+		return;
+	}
+
+	// support undo/redo
+	MediaPlate->SetFlags(RF_Transactional);
+	GEditor->RegisterForUndo(this);
+
+	BindCommands();
+
+	// create tab layout
+	const TSharedRef<FTabManager::FLayout> Layout = FTabManager::NewLayout("Standalone_MediaPlateEditor_v1")
+		->AddArea
+		(
+			FTabManager::NewPrimaryArea()
+				->SetOrientation(Orient_Horizontal)
+				->Split
+				(
+					FTabManager::NewSplitter()
+						->SetOrientation(Orient_Vertical)
+						->Split
+						(
+							// viewer
+							FTabManager::NewStack()
+								->AddTab(MediaPlateEditorToolkit::ViewerTabId, ETabState::OpenedTab)
+								->SetHideTabWell(true)
+								->SetSizeCoefficient(0.6f)
+						)
+						->Split
+						(
+							// media library
+							FTabManager::NewStack()
+								->AddTab(MediaPlateEditorToolkit::MediaTabId, ETabState::OpenedTab)
+								->SetHideTabWell(true)
+								->SetSizeCoefficient(0.3f)
+						)
+				)
+				->Split
+				(
+					FTabManager::NewSplitter()
+						->SetOrientation(Orient_Vertical)
+						->SetSizeCoefficient(0.33f)
+						->Split
+						(
+							// playlist
+							FTabManager::NewStack()
+								->AddTab(MediaPlateEditorToolkit::PlaylistTabId, ETabState::OpenedTab)
+								->SetSizeCoefficient(0.5f)
+						)
+						->Split
+						(
+							// details, info, stats
+							FTabManager::NewStack()
+								->AddTab(MediaPlateEditorToolkit::DetailsTabId, ETabState::OpenedTab)
+								->AddTab(MediaPlateEditorToolkit::InfoTabId, ETabState::OpenedTab)
+								->AddTab(MediaPlateEditorToolkit::StatsTabId, ETabState::ClosedTab)
+								->SetForegroundTab(MediaPlateEditorToolkit::DetailsTabId)
+								->SetSizeCoefficient(0.5f)
+						)
+				)
+		);
+
+	FAssetEditorToolkit::InitAssetEditor(
+		InMode,
+		InToolkitHost,
+		MediaPlateEditorToolkit::AppIdentifier,
+		Layout,
+		true,
+		true,
+		InMediaPlate
+	);
+	
+	ExtendToolBar();
+	RegenerateMenusAndToolbars();
+}
+
+/* FAssetEditorToolkit interface
+ *****************************************************************************/
+
+FString FMediaPlateEditorToolkit::GetDocumentationLink() const
+{
+	return FString(TEXT("WorkingWithMedia/IntegratingMedia/MediaFramework"));
+}
+
+void FMediaPlateEditorToolkit::OnClose()
+{
+	TObjectPtr<UMediaPlayer> MediaPlayer = MediaPlate->GetMediaPlayer();
+	if (MediaPlayer != nullptr)
+	{
+		MediaPlayer->Close();
+	}
+}
+
+void FMediaPlateEditorToolkit::RegisterTabSpawners(const TSharedRef<FTabManager>& InTabManager)
+{
+	WorkspaceMenuCategory = InTabManager->AddLocalWorkspaceMenuCategory(LOCTEXT("WorkspaceMenu_MediaPlateEditor", "Media Plate Editor"));
+	auto WorkspaceMenuCategoryRef = WorkspaceMenuCategory.ToSharedRef();
+
+	FAssetEditorToolkit::RegisterTabSpawners(InTabManager);
+
+	InTabManager->RegisterTabSpawner(MediaPlateEditorToolkit::DetailsTabId, FOnSpawnTab::CreateSP(this, &FMediaPlateEditorToolkit::HandleTabManagerSpawnTab, MediaPlateEditorToolkit::DetailsTabId))
+		.SetDisplayName(LOCTEXT("DetailsTabName", "Details"))
+		.SetGroup(WorkspaceMenuCategoryRef)
+		.SetIcon(FSlateIcon(FEditorStyle::GetStyleSetName(), "LevelEditor.Tabs.Details"));
+
+	InTabManager->RegisterTabSpawner(MediaPlateEditorToolkit::InfoTabId, FOnSpawnTab::CreateSP(this, &FMediaPlateEditorToolkit::HandleTabManagerSpawnTab, MediaPlateEditorToolkit::InfoTabId))
+		.SetDisplayName(LOCTEXT("InfoTabName", "Info"))
+		.SetGroup(WorkspaceMenuCategoryRef)
+		.SetIcon(FSlateIcon(Style->GetStyleSetName(), "MediaPlateEditor.Tabs.Info"));
+
+	InTabManager->RegisterTabSpawner(MediaPlateEditorToolkit::MediaTabId, FOnSpawnTab::CreateSP(this, &FMediaPlateEditorToolkit::HandleTabManagerSpawnTab, MediaPlateEditorToolkit::MediaTabId))
+		.SetDisplayName(LOCTEXT("MediaTabName", "Media Library"))
+		.SetGroup(WorkspaceMenuCategoryRef)
+		.SetIcon(FSlateIcon(Style->GetStyleSetName(), "MediaPlateEditor.Tabs.Media"));
+
+	InTabManager->RegisterTabSpawner(MediaPlateEditorToolkit::ViewerTabId, FOnSpawnTab::CreateSP(this, &FMediaPlateEditorToolkit::HandleTabManagerSpawnTab, MediaPlateEditorToolkit::ViewerTabId))
+		.SetDisplayName(LOCTEXT("PlayerTabName", "Player"))
+		.SetGroup(WorkspaceMenuCategoryRef)
+		.SetIcon(FSlateIcon(Style->GetStyleSetName(), "MediaPlateEditor.Tabs.Player"));
+
+	InTabManager->RegisterTabSpawner(MediaPlateEditorToolkit::PlaylistTabId, FOnSpawnTab::CreateSP(this, &FMediaPlateEditorToolkit::HandleTabManagerSpawnTab, MediaPlateEditorToolkit::PlaylistTabId))
+		.SetDisplayName(LOCTEXT("PlaylistTabName", "Playlist"))
+		.SetGroup(WorkspaceMenuCategoryRef)
+		.SetIcon(FSlateIcon(Style->GetStyleSetName(), "MediaPlateEditor.Tabs.Playlist"));
+
+	InTabManager->RegisterTabSpawner(MediaPlateEditorToolkit::StatsTabId, FOnSpawnTab::CreateSP(this, &FMediaPlateEditorToolkit::HandleTabManagerSpawnTab, MediaPlateEditorToolkit::StatsTabId))
+		.SetDisplayName(LOCTEXT("StatsTabName", "Stats"))
+		.SetGroup(WorkspaceMenuCategoryRef)
+		.SetIcon(FSlateIcon(Style->GetStyleSetName(), "MediaPlateEditor.Tabs.Stats"));
+}
+
+void FMediaPlateEditorToolkit::UnregisterTabSpawners(const TSharedRef<class FTabManager>& InTabManager)
+{
+	FAssetEditorToolkit::UnregisterTabSpawners(InTabManager);
+
+	InTabManager->UnregisterTabSpawner(MediaPlateEditorToolkit::ViewerTabId);
+	InTabManager->UnregisterTabSpawner(MediaPlateEditorToolkit::StatsTabId);
+	InTabManager->UnregisterTabSpawner(MediaPlateEditorToolkit::PlaylistTabId);
+	InTabManager->UnregisterTabSpawner(MediaPlateEditorToolkit::MediaTabId);
+	InTabManager->UnregisterTabSpawner(MediaPlateEditorToolkit::InfoTabId);
+	InTabManager->UnregisterTabSpawner(MediaPlateEditorToolkit::DetailsTabId);
+}
+
+/* IToolkit interface
+ *****************************************************************************/
+
+FText FMediaPlateEditorToolkit::GetBaseToolkitName() const
+{
+	return LOCTEXT("AppLabel", "Media Plate Editor");
+}
+
+FName FMediaPlateEditorToolkit::GetToolkitFName() const
+{
+	return FName("MediaPlateEditor");
+}
+
+FLinearColor FMediaPlateEditorToolkit::GetWorldCentricTabColorScale() const
+{
+	return FLinearColor(0.3f, 0.2f, 0.5f, 0.5f);
+}
+
+FString FMediaPlateEditorToolkit::GetWorldCentricTabPrefix() const
+{
+	return LOCTEXT("WorldCentricTabPrefix", "MediaPlate ").ToString();
+}
+
+/* FGCObject interface
+ *****************************************************************************/
+
+void FMediaPlateEditorToolkit::AddReferencedObjects(FReferenceCollector& Collector)
+{
+	Collector.AddReferencedObject(MediaPlate);
+}
+
+/* FEditorUndoClient interface
+*****************************************************************************/
+
+void FMediaPlateEditorToolkit::PostUndo(bool bSuccess)
+{
+	// do nothing
+}
+
+void FMediaPlateEditorToolkit::PostRedo(bool bSuccess)
+{
+	PostUndo(bSuccess);
+}
+
+/* FMediaPlayerEditorToolkit implementation
+ *****************************************************************************/
+
+void FMediaPlateEditorToolkit::BindCommands()
+{
+	const FMediaPlateEditorCommands& Commands = FMediaPlateEditorCommands::Get();
+
+	ToolkitCommands->MapAction(
+		Commands.CloseMedia,
+		FExecuteAction::CreateLambda([this] { MediaPlate->GetMediaPlayer()->Close(); }),
+		FCanExecuteAction::CreateLambda([this] { return !MediaPlate->GetMediaPlayer()->GetUrl().IsEmpty(); })
+	);
+
+	ToolkitCommands->MapAction(
+		Commands.ForwardMedia,
+		FExecuteAction::CreateLambda([this]{ MediaPlate->GetMediaPlayer()->SetRate(GetForwardRate()); }),
+		FCanExecuteAction::CreateLambda([this]{
+			TObjectPtr<UMediaPlayer> MediaPlayer = MediaPlate->GetMediaPlayer();
+			return MediaPlayer->IsReady() && MediaPlayer->SupportsRate(GetForwardRate(), false);
+		})
+	);
+
+	ToolkitCommands->MapAction(
+		Commands.NextMedia,
+		FExecuteAction::CreateLambda([this]{ MediaPlate->GetMediaPlayer()->Next(); }),
+		FCanExecuteAction::CreateLambda([this]{ return (MediaPlate->GetMediaPlayer()->GetPlaylistRef().Num() > 1); })
+	);
+
+	ToolkitCommands->MapAction(
+		Commands.OpenMedia,
+		FExecuteAction::CreateLambda([this] { MediaPlate->GetMediaPlayer()->OpenSource(MediaPlate->MediaSource); }),
+		FCanExecuteAction::CreateLambda([this] { return true; })
+	);
+
+	ToolkitCommands->MapAction(
+		Commands.PauseMedia,
+		FExecuteAction::CreateLambda([this]{ MediaPlate->GetMediaPlayer()->Pause(); }),
+		FCanExecuteAction::CreateLambda([this]{
+			TObjectPtr<UMediaPlayer> MediaPlayer = MediaPlate->GetMediaPlayer();
+			return MediaPlayer->CanPause() && !MediaPlayer->IsPaused();
+		})
+	);
+
+	ToolkitCommands->MapAction(
+		Commands.PlayMedia,
+		FExecuteAction::CreateLambda([this]{ MediaPlate->GetMediaPlayer()->Play(); }),
+		FCanExecuteAction::CreateLambda([this]{
+			TObjectPtr<UMediaPlayer> MediaPlayer = MediaPlate->GetMediaPlayer(); 
+			return MediaPlayer->IsReady() && (!MediaPlayer->IsPlaying() || (MediaPlayer->GetRate() != 1.0f));
+		})
+	);
+
+	ToolkitCommands->MapAction(
+		Commands.PreviousMedia,
+		FExecuteAction::CreateLambda([this]{ MediaPlate->GetMediaPlayer()->Previous(); }),
+		FCanExecuteAction::CreateLambda([this]{ return (MediaPlate->GetMediaPlayer()->GetPlaylistRef().Num() > 1); })
+	);
+
+	ToolkitCommands->MapAction(
+		Commands.ReverseMedia,
+		FExecuteAction::CreateLambda([this]{ MediaPlate->GetMediaPlayer()->SetRate(GetReverseRate()); } ),
+		FCanExecuteAction::CreateLambda([this]{
+			TObjectPtr<UMediaPlayer> MediaPlayer = MediaPlate->GetMediaPlayer();
+			return MediaPlayer->IsReady() && MediaPlayer->SupportsRate(GetReverseRate(), false);
+		})
+	);
+
+	ToolkitCommands->MapAction(
+		Commands.RewindMedia,
+		FExecuteAction::CreateLambda([this]{ MediaPlate->GetMediaPlayer()->Rewind(); }),
+		FCanExecuteAction::CreateLambda([this]{
+			TObjectPtr<UMediaPlayer> MediaPlayer = MediaPlate->GetMediaPlayer(); 
+			return MediaPlate->GetMediaPlayer()->IsReady() && MediaPlayer->SupportsSeeking() && MediaPlayer->GetTime() > FTimespan::Zero();
+		})
+	);
+}
+
+BEGIN_SLATE_FUNCTION_BUILD_OPTIMIZATION
+void FMediaPlateEditorToolkit::ExtendToolBar()
+{
+	struct Local
+	{
+		static void FillToolbar(FToolBarBuilder& ToolbarBuilder, const TSharedRef<FUICommandList> ToolkitCommands)
+		{
+			ToolbarBuilder.BeginSection("PlaybackControls");
+			{
+				ToolbarBuilder.AddToolBarButton(FMediaPlateEditorCommands::Get().PreviousMedia);
+				ToolbarBuilder.AddToolBarButton(FMediaPlateEditorCommands::Get().RewindMedia);
+				ToolbarBuilder.AddToolBarButton(FMediaPlateEditorCommands::Get().ReverseMedia);
+				ToolbarBuilder.AddToolBarButton(FMediaPlateEditorCommands::Get().PlayMedia);
+				ToolbarBuilder.AddToolBarButton(FMediaPlateEditorCommands::Get().PauseMedia);
+				ToolbarBuilder.AddToolBarButton(FMediaPlateEditorCommands::Get().ForwardMedia);
+				ToolbarBuilder.AddToolBarButton(FMediaPlateEditorCommands::Get().NextMedia);
+			}
+			ToolbarBuilder.EndSection();
+
+			ToolbarBuilder.BeginSection("MediaControls");
+			{
+				ToolbarBuilder.AddToolBarButton(FMediaPlateEditorCommands::Get().OpenMedia);
+				ToolbarBuilder.AddToolBarButton(FMediaPlateEditorCommands::Get().CloseMedia);
+			}
+			ToolbarBuilder.EndSection();
+		}
+	};
+
+	TSharedPtr<FExtender> ToolbarExtender = MakeShareable(new FExtender);
+
+	ToolbarExtender->AddToolBarExtension(
+		"Asset",
+		EExtensionHook::After,
+		GetToolkitCommands(),
+		FToolBarExtensionDelegate::CreateStatic(&Local::FillToolbar, GetToolkitCommands())
+	);
+
+	AddToolbarExtender(ToolbarExtender);
+}
+END_SLATE_FUNCTION_BUILD_OPTIMIZATION
+
+float FMediaPlateEditorToolkit::GetForwardRate() const
+{
+	TObjectPtr<UMediaPlayer> MediaPlayer = MediaPlate->GetMediaPlayer();
+	float Rate = MediaPlayer->GetRate();
+
+	if (Rate < 1.0f)
+	{
+		Rate = 1.0f;
+	}
+
+	return 2.0f * Rate;
+}
+
+float FMediaPlateEditorToolkit::GetReverseRate() const
+{
+	TObjectPtr<UMediaPlayer> MediaPlayer = MediaPlate->GetMediaPlayer();
+	float Rate = MediaPlayer->GetRate();
+
+	if (Rate > -1.0f)
+	{
+		return -1.0f;
+	}
+
+	return 2.0f * Rate;
+}
+
+/* FMediaPlayerEditorToolkit callbacks
+ *****************************************************************************/
+
+TSharedRef<SDockTab> FMediaPlateEditorToolkit::HandleTabManagerSpawnTab(const FSpawnTabArgs& Args, FName TabIdentifier)
+{
+	UMediaPlayer* MediaPlayer = MediaPlate->GetMediaPlayer();
+	TSharedPtr<SWidget> TabWidget = SNullWidget::NullWidget;
+	if (TabIdentifier == MediaPlateEditorToolkit::ViewerTabId)
+	{
+		TabWidget = SNew(SMediaPlayerEditorViewer, *MediaPlayer, Style);
+	}
+#if false
+	if (TabIdentifier == MediaPlayerEditorToolkit::DetailsTabId)
+	{
+		TabWidget = SNew(SMediaPlayerEditorDetails, *MediaPlayer, Style);
+	}
+	else if (TabIdentifier == MediaPlayerEditorToolkit::InfoTabId)
+	{
+		TabWidget = SNew(SMediaPlayerEditorInfo, *MediaPlayer, Style);
+	}
+	else if (TabIdentifier == MediaPlayerEditorToolkit::MediaTabId)
+	{
+		TabWidget = SNew(SMediaPlayerEditorMedia, *MediaPlayer, Style);
+	}
+	else if (TabIdentifier == MediaPlayerEditorToolkit::PlaylistTabId)
+	{
+		TabWidget = SNew(SMediaPlayerEditorPlaylist, *MediaPlayer, Style);
+	}
+	else if (TabIdentifier == MediaPlayerEditorToolkit::StatsTabId)
+	{
+		TabWidget = SNew(SMediaPlayerEditorStats, *MediaPlayer, Style);
+	}
+	else if (TabIdentifier == MediaPlayerEditorToolkit::ViewerTabId)
+	{
+		TabWidget = SNew(SMediaPlayerEditorViewer, *MediaPlayer, Style);
+	}
+#endif
+	return SNew(SDockTab)
+		.TabRole(ETabRole::PanelTab)
+		[
+			TabWidget.ToSharedRef()
+		];
+}
+
+#undef LOCTEXT_NAMESPACE
diff --git a/Engine/Plugins/Media/MediaPlate/Source/MediaPlateEditor/Private/Toolkits/MediaPlateEditorToolkit.h b/Engine/Plugins/Media/MediaPlate/Source/MediaPlateEditor/Private/Toolkits/MediaPlateEditorToolkit.h
new file mode 100644
--- /dev/null
+++ b/Engine/Plugins/Media/MediaPlate/Source/MediaPlateEditor/Private/Toolkits/MediaPlateEditorToolkit.h
@@ -0,0 +1,97 @@
+// Copyright Epic Games, Inc. All Rights Reserved.
+
+#pragma once
+
+#include "CoreMinimal.h"
+#include "EditorUndoClient.h"
+#include "Framework/Docking/TabManager.h"
+#include "Styling/ISlateStyle.h"
+#include "Toolkits/AssetEditorToolkit.h"
+#include "UObject/GCObject.h"
+
+class AMediaPlate;
+
+/**
+ * Implements an Editor toolkit for media plates.
+ */
+class FMediaPlateEditorToolkit
+	: public FAssetEditorToolkit
+	, public FEditorUndoClient
+	, public FGCObject
+{
+public:
+	/**
+	 * Creates and initializes a new instance.
+	 *
+	 * @param InStyle The style set to use.
+	 */
+	FMediaPlateEditorToolkit(const TSharedRef<ISlateStyle>& InStyle);
+
+	/** Virtual destructor. */
+	virtual ~FMediaPlateEditorToolkit();
+
+	/**
+	 * Initializes the editor tool kit.
+	 *
+	 * @param InMediaPlate The AMediaPlate asset to edit.
+	 * @param InMode The mode to create the toolkit in.
+	 * @param InToolkitHost The toolkit host.
+	 */
+	void Initialize(AMediaPlate* InMediaPlate, const EToolkitMode::Type InMode, const TSharedPtr<IToolkitHost>& InToolkitHost);
+
+	//~ FAssetEditorToolkit interface
+	virtual FString GetDocumentationLink() const override;
+	virtual void OnClose() override;
+	virtual void RegisterTabSpawners(const TSharedRef<FTabManager>& TabManager) override;
+	virtual void UnregisterTabSpawners(const TSharedRef<FTabManager>& TabManager) override;
+
+	//~ IToolkit interface
+	virtual FText GetBaseToolkitName() const override;
+	virtual FName GetToolkitFName() const override;
+	virtual FLinearColor GetWorldCentricTabColorScale() const override;
+	virtual FString GetWorldCentricTabPrefix() const override;
+
+	//~ FGCObject interface
+	virtual void AddReferencedObjects(FReferenceCollector& Collector) override;
+	virtual FString GetReferencerName() const override
+	{
+		return TEXT("FMediaPlateEditorToolkit");
+	}
+
+protected:
+
+	//~ FEditorUndoClient interface
+	virtual void PostUndo(bool bSuccess) override;
+	virtual void PostRedo(bool bSuccess) override;
+
+	/** Binds the UI commands to delegates. */
+	void BindCommands();
+
+	/** Builds the toolbar widget for the media player editor. */
+	void ExtendToolBar();
+
+	/**
+	 * Gets the playback rate for fast forward.
+	 *
+	 * @return Forward playback rate.
+	 */
+	float GetForwardRate() const;
+
+	/**
+	 * Gets the playback rate for reverse.
+	 *
+	 * @return Reverse playback rate.
+	 */
+	float GetReverseRate() const;
+
+private:
+
+	/** Callback for spawning tabs. */
+	TSharedRef<SDockTab> HandleTabManagerSpawnTab(const FSpawnTabArgs& Args, FName TabIdentifier);
+
+	/** The media plate asset being edited. */
+	AMediaPlate* MediaPlate;
+
+	/** Pointer to the style set to use for toolkits. */
+	TSharedRef<ISlateStyle> Style;
+};
diff --git a/Engine/Plugins/Media/MediaPlate/Source/MediaPlateEditor/Public/MediaPlateEditorModule.h b/Engine/Plugins/Media/MediaPlate/Source/MediaPlateEditor/Public/MediaPlateEditorModule.h
--- a/Engine/Plugins/Media/MediaPlate/Source/MediaPlateEditor/Public/MediaPlateEditorModule.h
+++ b/Engine/Plugins/Media/MediaPlate/Source/MediaPlateEditor/Public/MediaPlateEditorModule.h
@@ -9,6 +9,9 @@
 #include "UObject/ObjectPtr.h"
 
 class AMediaPlate;
+class IAssetTools;
+class IAssetTypeActions;
+class ISlateStyle;
 
 /** Log category for this module. */
 DECLARE_LOG_CATEGORY_EXTERN(LogMediaPlateEditor, Log, All);
@@ -36,4 +39,27 @@ private:
 	FName MediaPlateName;
 	/** Holds all the media plates that are playing. */
 	TArray<TWeakObjectPtr<AMediaPlate>> ActiveMediaPlates;
+
+	/** Holds the plug-ins style set. */
+	TSharedPtr<ISlateStyle> Style;
+	/** The collection of registered asset type actions. */
+	TArray<TSharedRef<IAssetTypeActions>> RegisteredAssetTypeActions;
+
+	/**
+	 * Registers all of our asset tools.
+	 */
+	void RegisterAssetTools();
+
+	/**
+	 * Registers a single asset type action.
+	 *
+	 * @param AssetTools	The asset tools object to register with.
+	 * @param Action		The asset type action to register.
+	 */
+	void RegisterAssetTypeAction(IAssetTools& AssetTools, TSharedRef<IAssetTypeActions> Action);
+
+	/**
+	 * Unregisters all of our asset tools.
+	 */
+	void UnregisterAssetTools();
 };
