MediaPlate: Removed unused tabs from the editor toolkit.

#rnx

diff --git a/Engine/Plugins/Media/MediaPlate/Source/MediaPlateEditor/Private/Toolkits/MediaPlateEditorToolkit.cpp b/Engine/Plugins/Media/MediaPlate/Source/MediaPlateEditor/Private/Toolkits/MediaPlateEditorToolkit.cpp
--- a/Engine/Plugins/Media/MediaPlate/Source/MediaPlateEditor/Private/Toolkits/MediaPlateEditorToolkit.cpp
+++ b/Engine/Plugins/Media/MediaPlate/Source/MediaPlateEditor/Private/Toolkits/MediaPlateEditorToolkit.cpp
@@ -19,11 +19,6 @@
 namespace MediaPlateEditorToolkit
 {
 	static const FName AppIdentifier("MediaPlateEditorApp");
-	static const FName DetailsTabId("Details");
-	static const FName InfoTabId("Info");
-	static const FName MediaTabId("Media");
-	static const FName PlaylistTabId("Playlist");
-	static const FName StatsTabId("Stats");
 	static const FName ViewerTabId("Viewer");
 }
 
@@ -63,7 +58,7 @@ void FMediaPlateEditorToolkit::Initialize(AMediaPlate* InMediaPlate, const ETool
 	BindCommands();
 
 	// create tab layout
-	const TSharedRef<FTabManager::FLayout> Layout = FTabManager::NewLayout("Standalone_MediaPlateEditor_v1")
+	const TSharedRef<FTabManager::FLayout> Layout = FTabManager::NewLayout("Standalone_MediaPlateEditor_v1.1")
 		->AddArea
 		(
 			FTabManager::NewPrimaryArea()
@@ -79,38 +74,7 @@ void FMediaPlateEditorToolkit::Initialize(AMediaPlate* InMediaPlate, const ETool
 								->AddTab(MediaPlateEditorToolkit::ViewerTabId, ETabState::OpenedTab)
 								->SetHideTabWell(true)
 								->SetSizeCoefficient(0.6f)
-						)
-						->Split
-						(
-							// media library
-							FTabManager::NewStack()
-								->AddTab(MediaPlateEditorToolkit::MediaTabId, ETabState::OpenedTab)
-								->SetHideTabWell(true)
-								->SetSizeCoefficient(0.3f)
-						)
-				)
-				->Split
-				(
-					FTabManager::NewSplitter()
-						->SetOrientation(Orient_Vertical)
-						->SetSizeCoefficient(0.33f)
-						->Split
-						(
-							// playlist
-							FTabManager::NewStack()
-								->AddTab(MediaPlateEditorToolkit::PlaylistTabId, ETabState::OpenedTab)
-								->SetSizeCoefficient(0.5f)
-						)
-						->Split
-						(
-							// details, info, stats
-							FTabManager::NewStack()
-								->AddTab(MediaPlateEditorToolkit::DetailsTabId, ETabState::OpenedTab)
-								->AddTab(MediaPlateEditorToolkit::InfoTabId, ETabState::OpenedTab)
-								->AddTab(MediaPlateEditorToolkit::StatsTabId, ETabState::ClosedTab)
-								->SetForegroundTab(MediaPlateEditorToolkit::DetailsTabId)
-								->SetSizeCoefficient(0.5f)
-						)
+						)	
 				)
 		);
 
@@ -152,35 +116,10 @@ void FMediaPlateEditorToolkit::RegisterTabSpawners(const TSharedRef<FTabManager>
 
 	FAssetEditorToolkit::RegisterTabSpawners(InTabManager);
 
-	InTabManager->RegisterTabSpawner(MediaPlateEditorToolkit::DetailsTabId, FOnSpawnTab::CreateSP(this, &FMediaPlateEditorToolkit::HandleTabManagerSpawnTab, MediaPlateEditorToolkit::DetailsTabId))
-		.SetDisplayName(LOCTEXT("DetailsTabName", "Details"))
-		.SetGroup(WorkspaceMenuCategoryRef)
-		.SetIcon(FSlateIcon(FEditorStyle::GetStyleSetName(), "LevelEditor.Tabs.Details"));
-
-	InTabManager->RegisterTabSpawner(MediaPlateEditorToolkit::InfoTabId, FOnSpawnTab::CreateSP(this, &FMediaPlateEditorToolkit::HandleTabManagerSpawnTab, MediaPlateEditorToolkit::InfoTabId))
-		.SetDisplayName(LOCTEXT("InfoTabName", "Info"))
-		.SetGroup(WorkspaceMenuCategoryRef)
-		.SetIcon(FSlateIcon(Style->GetStyleSetName(), "MediaPlateEditor.Tabs.Info"));
-
-	InTabManager->RegisterTabSpawner(MediaPlateEditorToolkit::MediaTabId, FOnSpawnTab::CreateSP(this, &FMediaPlateEditorToolkit::HandleTabManagerSpawnTab, MediaPlateEditorToolkit::MediaTabId))
-		.SetDisplayName(LOCTEXT("MediaTabName", "Media Library"))
-		.SetGroup(WorkspaceMenuCategoryRef)
-		.SetIcon(FSlateIcon(Style->GetStyleSetName(), "MediaPlateEditor.Tabs.Media"));
-
 	InTabManager->RegisterTabSpawner(MediaPlateEditorToolkit::ViewerTabId, FOnSpawnTab::CreateSP(this, &FMediaPlateEditorToolkit::HandleTabManagerSpawnTab, MediaPlateEditorToolkit::ViewerTabId))
 		.SetDisplayName(LOCTEXT("PlayerTabName", "Player"))
 		.SetGroup(WorkspaceMenuCategoryRef)
 		.SetIcon(FSlateIcon(Style->GetStyleSetName(), "MediaPlateEditor.Tabs.Player"));
-
-	InTabManager->RegisterTabSpawner(MediaPlateEditorToolkit::PlaylistTabId, FOnSpawnTab::CreateSP(this, &FMediaPlateEditorToolkit::HandleTabManagerSpawnTab, MediaPlateEditorToolkit::PlaylistTabId))
-		.SetDisplayName(LOCTEXT("PlaylistTabName", "Playlist"))
-		.SetGroup(WorkspaceMenuCategoryRef)
-		.SetIcon(FSlateIcon(Style->GetStyleSetName(), "MediaPlateEditor.Tabs.Playlist"));
-
-	InTabManager->RegisterTabSpawner(MediaPlateEditorToolkit::StatsTabId, FOnSpawnTab::CreateSP(this, &FMediaPlateEditorToolkit::HandleTabManagerSpawnTab, MediaPlateEditorToolkit::StatsTabId))
-		.SetDisplayName(LOCTEXT("StatsTabName", "Stats"))
-		.SetGroup(WorkspaceMenuCategoryRef)
-		.SetIcon(FSlateIcon(Style->GetStyleSetName(), "MediaPlateEditor.Tabs.Stats"));
 }
 
 void FMediaPlateEditorToolkit::UnregisterTabSpawners(const TSharedRef<class FTabManager>& InTabManager)
@@ -188,11 +127,6 @@ void FMediaPlateEditorToolkit::UnregisterTabSpawners(const TSharedRef<class FTab
 	FAssetEditorToolkit::UnregisterTabSpawners(InTabManager);
 
 	InTabManager->UnregisterTabSpawner(MediaPlateEditorToolkit::ViewerTabId);
-	InTabManager->UnregisterTabSpawner(MediaPlateEditorToolkit::StatsTabId);
-	InTabManager->UnregisterTabSpawner(MediaPlateEditorToolkit::PlaylistTabId);
-	InTabManager->UnregisterTabSpawner(MediaPlateEditorToolkit::MediaTabId);
-	InTabManager->UnregisterTabSpawner(MediaPlateEditorToolkit::InfoTabId);
-	InTabManager->UnregisterTabSpawner(MediaPlateEditorToolkit::DetailsTabId);
 }
 
 /* IToolkit interface
@@ -394,32 +328,7 @@ TSharedRef<SDockTab> FMediaPlateEditorToolkit::HandleTabManagerSpawnTab(const FS
 	{
 		TabWidget = SNew(SMediaPlayerEditorViewer, *MediaPlayer, Style);
 	}
-#if false
-	if (TabIdentifier == MediaPlayerEditorToolkit::DetailsTabId)
-	{
-		TabWidget = SNew(SMediaPlayerEditorDetails, *MediaPlayer, Style);
-	}
-	else if (TabIdentifier == MediaPlayerEditorToolkit::InfoTabId)
-	{
-		TabWidget = SNew(SMediaPlayerEditorInfo, *MediaPlayer, Style);
-	}
-	else if (TabIdentifier == MediaPlayerEditorToolkit::MediaTabId)
-	{
-		TabWidget = SNew(SMediaPlayerEditorMedia, *MediaPlayer, Style);
-	}
-	else if (TabIdentifier == MediaPlayerEditorToolkit::PlaylistTabId)
-	{
-		TabWidget = SNew(SMediaPlayerEditorPlaylist, *MediaPlayer, Style);
-	}
-	else if (TabIdentifier == MediaPlayerEditorToolkit::StatsTabId)
-	{
-		TabWidget = SNew(SMediaPlayerEditorStats, *MediaPlayer, Style);
-	}
-	else if (TabIdentifier == MediaPlayerEditorToolkit::ViewerTabId)
-	{
-		TabWidget = SNew(SMediaPlayerEditorViewer, *MediaPlayer, Style);
-	}
-#endif
+
 	return SNew(SDockTab)
 		.TabRole(ETabRole::PanelTab)
 		[
