MediaPlate: Added looping option.
#rnx

diff --git a/Engine/Plugins/Media/MediaPlate/Source/MediaPlate/Private/MediaPlate.cpp b/Engine/Plugins/Media/MediaPlate/Source/MediaPlate/Private/MediaPlate.cpp
--- a/Engine/Plugins/Media/MediaPlate/Source/MediaPlate/Private/MediaPlate.cpp
+++ b/Engine/Plugins/Media/MediaPlate/Source/MediaPlate/Private/MediaPlate.cpp
@@ -133,6 +133,7 @@ void AMediaPlate::Play()
 	if (MediaPlayer != nullptr)
 	{
 		MediaPlayer->PlayOnOpen = true;
+		MediaPlayer->SetLooping(bLoop);
 		MediaPlayer->OpenSource(MediaSource);
 	}
 }
diff --git a/Engine/Plugins/Media/MediaPlate/Source/MediaPlate/Public/MediaPlate.h b/Engine/Plugins/Media/MediaPlate/Source/MediaPlate/Public/MediaPlate.h
--- a/Engine/Plugins/Media/MediaPlate/Source/MediaPlate/Public/MediaPlate.h
+++ b/Engine/Plugins/Media/MediaPlate/Source/MediaPlate/Public/MediaPlate.h
@@ -55,6 +55,10 @@ public:
 	UPROPERTY(EditAnywhere, Category = "MediaPlate")
 	bool bAutoPlay;
 
+	/** If set then loop when we reach the end. */
+	UPROPERTY(EditAnywhere, Category = "MediaPlate")
+	bool bLoop;
+
 	/** Holds the media player. */
 	UPROPERTY(Category = MediaPlate, VisibleAnywhere, BlueprintReadOnly)
 	TObjectPtr<UMediaComponent> MediaComponent;
