MediaPlate: Added support for urls/files.

#rnx

diff --git a/Engine/Plugins/Media/MediaPlate/Source/MediaPlate/Private/MediaPlate.cpp b/Engine/Plugins/Media/MediaPlate/Source/MediaPlate/Private/MediaPlate.cpp
--- a/Engine/Plugins/Media/MediaPlate/Source/MediaPlate/Private/MediaPlate.cpp
+++ b/Engine/Plugins/Media/MediaPlate/Source/MediaPlate/Private/MediaPlate.cpp
@@ -129,12 +129,39 @@ UMediaTexture* AMediaPlate::GetMediaTexture()
 
 void AMediaPlate::Play()
 {
+	bool bIsPlaying = false;
 	TObjectPtr<UMediaPlayer> MediaPlayer = GetMediaPlayer();
 	if (MediaPlayer != nullptr)
 	{
 		MediaPlayer->PlayOnOpen = true;
 		MediaPlayer->SetLooping(bLoop);
-		MediaPlayer->OpenSource(MediaSource);
+		
+		// Try and play the URL.
+		if (Url.FilePath.IsEmpty() == false)
+		{
+			// Is it a URL or a file?
+			bool bIsUrl = Url.FilePath.Contains(TEXT("://"));
+			if (bIsUrl)
+			{
+				bIsPlaying = MediaPlayer->OpenUrl(Url.FilePath);
+			}
+			else
+			{
+				bIsPlaying = MediaPlayer->OpenFile(Url.FilePath);
+			}
+		}
+
+		// If we did not get anything, try the media source.
+		if (bIsPlaying == false)
+		{
+			bIsPlaying = MediaPlayer->OpenSource(MediaSource);
+		}
+	}
+
+	// Did anything play?
+	if (bIsPlaying == false)
+	{
+		UE_LOG(LogMediaPlate, Warning, TEXT("Could not play anything."));
 	}
 }
 
diff --git a/Engine/Plugins/Media/MediaPlate/Source/MediaPlate/Public/MediaPlate.h b/Engine/Plugins/Media/MediaPlate/Source/MediaPlate/Public/MediaPlate.h
--- a/Engine/Plugins/Media/MediaPlate/Source/MediaPlate/Public/MediaPlate.h
+++ b/Engine/Plugins/Media/MediaPlate/Source/MediaPlate/Public/MediaPlate.h
@@ -68,10 +68,14 @@ public:
 	UPROPERTY(Category = MediaPlate, VisibleAnywhere, BlueprintReadOnly)
 	TObjectPtr<UStaticMeshComponent> StaticMeshComponent;
 
-	/** What media to play. */
+	/** What media to play. This will only be used if Url is not valid. */
 	UPROPERTY(EditAnywhere, Category = MediaPlate)
 	TObjectPtr<UMediaSource> MediaSource;
 
+	/** URL (or file)  to play. This will take precedence over the MediaSource. */
+	UPROPERTY(EditAnywhere, Category = MediaPlate)
+	FFilePath Url;
+
 private:
 	/** Name for our media component. */
 	static FLazyName MediaComponentName;
