MediaPlate: Fixed missing button images in the media plate window.
#rnx

diff --git a/Engine/Plugins/Media/MediaPlate/Source/MediaPlateEditor/Private/MediaPlateEditorStyle.cpp b/Engine/Plugins/Media/MediaPlate/Source/MediaPlateEditor/Private/MediaPlateEditorStyle.cpp
--- a/Engine/Plugins/Media/MediaPlate/Source/MediaPlateEditor/Private/MediaPlateEditorStyle.cpp
+++ b/Engine/Plugins/Media/MediaPlate/Source/MediaPlateEditor/Private/MediaPlateEditorStyle.cpp
@@ -32,10 +32,10 @@ FMediaPlateEditorStyle::FMediaPlateEditorStyle()
 	SetContentRoot(FPaths::EnginePluginsDir() / TEXT("Media/MediaPlayerEditor/Content"));
 
 	// buttons
-	Set("MediaPlateEditor.SourceButton", new IMAGE_BRUSH("btn_source_12x", Icon12x12));
-	Set("MediaPlateEditor.GoButton", new IMAGE_BRUSH("btn_go_12x", Icon12x12));
-	Set("MediaPlateEditor.ReloadButton", new IMAGE_BRUSH("btn_reload_12x", Icon12x12));
-	Set("MediaPlateEditor.SettingsButton", new IMAGE_BRUSH("btn_settings_16x", Icon12x12));
+	Set("MediaPlayerEditor.SourceButton", new IMAGE_BRUSH("btn_source_12x", Icon12x12));
+	Set("MediaPlayerEditor.GoButton", new IMAGE_BRUSH("btn_go_12x", Icon12x12));
+	Set("MediaPlayerEditor.ReloadButton", new IMAGE_BRUSH("btn_reload_12x", Icon12x12));
+	Set("MediaPlayerEditor.SettingsButton", new IMAGE_BRUSH("btn_settings_16x", Icon12x12));
 
 	// misc
 	Set("MediaPlateEditor.DragDropBorder", new BOX_BRUSH("border_dragdrop", 0.5f));
