MediaPlate: Added details tab to media plate toolkit.
#rnx

diff --git a/Engine/Plugins/Media/MediaPlate/Source/MediaPlateEditor/Private/Toolkits/MediaPlateEditorToolkit.cpp b/Engine/Plugins/Media/MediaPlate/Source/MediaPlateEditor/Private/Toolkits/MediaPlateEditorToolkit.cpp
--- a/Engine/Plugins/Media/MediaPlate/Source/MediaPlateEditor/Private/Toolkits/MediaPlateEditorToolkit.cpp
+++ b/Engine/Plugins/Media/MediaPlate/Source/MediaPlateEditor/Private/Toolkits/MediaPlateEditorToolkit.cpp
@@ -13,6 +13,7 @@
 #include "Models/MediaPlateEditorCommands.h"
 #include "SlateOptMacros.h"
 #include "Widgets/Docking/SDockTab.h"
+#include "Widgets/SMediaPlateEditorDetails.h"
 #include "Widgets/SMediaPlayerEditorViewer.h"
 
 #define LOCTEXT_NAMESPACE "FMediaPlateEditorToolkit"
@@ -20,6 +21,7 @@
 namespace MediaPlateEditorToolkit
 {
 	static const FName AppIdentifier("MediaPlateEditorApp");
+	static const FName DetailsTabId("Details");
 	static const FName ViewerTabId("Viewer");
 }
 
@@ -59,7 +61,7 @@ void FMediaPlateEditorToolkit::Initialize(AMediaPlate* InMediaPlate, const ETool
 	BindCommands();
 
 	// create tab layout
-	const TSharedRef<FTabManager::FLayout> Layout = FTabManager::NewLayout("Standalone_MediaPlateEditor_v1.1")
+	const TSharedRef<FTabManager::FLayout> Layout = FTabManager::NewLayout("Standalone_MediaPlateEditor_v1.1b")
 		->AddArea
 		(
 			FTabManager::NewPrimaryArea()
@@ -67,7 +69,7 @@ void FMediaPlateEditorToolkit::Initialize(AMediaPlate* InMediaPlate, const ETool
 				->Split
 				(
 					FTabManager::NewSplitter()
-						->SetOrientation(Orient_Vertical)
+						->SetOrientation(Orient_Horizontal)
 						->Split
 						(
 							// viewer
@@ -76,6 +78,13 @@ void FMediaPlateEditorToolkit::Initialize(AMediaPlate* InMediaPlate, const ETool
 								->SetHideTabWell(true)
 								->SetSizeCoefficient(0.6f)
 						)	
+						->Split
+						(
+							// viewer
+							FTabManager::NewStack()
+							->AddTab(MediaPlateEditorToolkit::DetailsTabId, ETabState::OpenedTab)
+							->SetSizeCoefficient(0.4f)
+						)
 				)
 		);
 
@@ -124,6 +133,13 @@ void FMediaPlateEditorToolkit::RegisterTabSpawners(const TSharedRef<FTabManager>
 
 	FAssetEditorToolkit::RegisterTabSpawners(InTabManager);
 
+	// Details tab.
+	InTabManager->RegisterTabSpawner(MediaPlateEditorToolkit::DetailsTabId, FOnSpawnTab::CreateSP(this, &FMediaPlateEditorToolkit::HandleTabManagerSpawnTab, MediaPlateEditorToolkit::DetailsTabId))
+		.SetDisplayName(LOCTEXT("DetailsTabName", "Details"))
+		.SetGroup(WorkspaceMenuCategoryRef)
+		.SetIcon(FSlateIcon(FEditorStyle::GetStyleSetName(), "LevelEditor.Tabs.Details"));
+
+	// Viewer tab.
 	InTabManager->RegisterTabSpawner(MediaPlateEditorToolkit::ViewerTabId, FOnSpawnTab::CreateSP(this, &FMediaPlateEditorToolkit::HandleTabManagerSpawnTab, MediaPlateEditorToolkit::ViewerTabId))
 		.SetDisplayName(LOCTEXT("PlayerTabName", "Player"))
 		.SetGroup(WorkspaceMenuCategoryRef)
@@ -135,6 +151,7 @@ void FMediaPlateEditorToolkit::UnregisterTabSpawners(const TSharedRef<class FTab
 	FAssetEditorToolkit::UnregisterTabSpawners(InTabManager);
 
 	InTabManager->UnregisterTabSpawner(MediaPlateEditorToolkit::ViewerTabId);
+	InTabManager->UnregisterTabSpawner(MediaPlateEditorToolkit::DetailsTabId);
 }
 
 /* IToolkit interface
@@ -333,7 +350,12 @@ TSharedRef<SDockTab> FMediaPlateEditorToolkit::HandleTabManagerSpawnTab(const FS
 	UMediaPlayer* MediaPlayer = MediaPlate->GetMediaPlayer();
 	UMediaTexture* MediaTexture = MediaPlate->GetMediaTexture();
 	TSharedPtr<SWidget> TabWidget = SNullWidget::NullWidget;
-	if (TabIdentifier == MediaPlateEditorToolkit::ViewerTabId)
+
+	if (TabIdentifier == MediaPlateEditorToolkit::DetailsTabId)
+	{
+		TabWidget = SNew(SMediaPlateEditorDetails, *MediaPlate, Style);
+	}
+	else if (TabIdentifier == MediaPlateEditorToolkit::ViewerTabId)
 	{
 		TabWidget = SNew(SMediaPlayerEditorViewer, *MediaPlayer, MediaTexture, Style);
 	}
diff --git a/Engine/Plugins/Media/MediaPlate/Source/MediaPlateEditor/Private/Widgets/SMediaPlateEditorDetails.cpp b/Engine/Plugins/Media/MediaPlate/Source/MediaPlateEditor/Private/Widgets/SMediaPlateEditorDetails.cpp
new file mode 100644
--- /dev/null
+++ b/Engine/Plugins/Media/MediaPlate/Source/MediaPlateEditor/Private/Widgets/SMediaPlateEditorDetails.cpp
@@ -0,0 +1,40 @@
+// Copyright Epic Games, Inc. All Rights Reserved.
+
+#include "Widgets/SMediaPlateEditorDetails.h"
+#include "IDetailsView.h"
+#include "MediaPlate.h"
+#include "Modules/ModuleManager.h"
+#include "PropertyEditorModule.h"
+
+#define LOCTEXT_NAMESPACE "SMediaPlateEditorDetails"
+
+/* SMediaPlateEditorDetails interface
+ *****************************************************************************/
+
+void SMediaPlateEditorDetails::Construct(const FArguments& InArgs, AMediaPlate& InMediaPlate, const TSharedRef<ISlateStyle>& InStyle)
+{
+	MediaPlate = &InMediaPlate;
+
+	// Initialize details view
+	FDetailsViewArgs DetailsViewArgs;
+	{
+		DetailsViewArgs.bAllowSearch = true;
+		DetailsViewArgs.bHideSelectionTip = true;
+		DetailsViewArgs.bLockable = false;
+		DetailsViewArgs.bSearchInitialKeyFocus = true;
+		DetailsViewArgs.bUpdatesFromSelection = false;
+		DetailsViewArgs.bShowOptions = true;
+		DetailsViewArgs.bShowModifiedPropertiesOption = false;
+	}
+
+	TSharedPtr<IDetailsView> DetailsView = FModuleManager::GetModuleChecked<FPropertyEditorModule>("PropertyEditor").CreateDetailView(DetailsViewArgs);
+	DetailsView->SetObject(MediaPlate);
+
+	ChildSlot
+	[
+		DetailsView.ToSharedRef()
+	];
+}
+
+
+#undef LOCTEXT_NAMESPACE
diff --git a/Engine/Plugins/Media/MediaPlate/Source/MediaPlateEditor/Private/Widgets/SMediaPlateEditorDetails.h b/Engine/Plugins/Media/MediaPlate/Source/MediaPlateEditor/Private/Widgets/SMediaPlateEditorDetails.h
new file mode 100644
--- /dev/null
+++ b/Engine/Plugins/Media/MediaPlate/Source/MediaPlateEditor/Private/Widgets/SMediaPlateEditorDetails.h
@@ -0,0 +1,38 @@
+// Copyright Epic Games, Inc. All Rights Reserved.
+
+#pragma once
+
+#include "CoreMinimal.h"
+#include "Styling/ISlateStyle.h"
+#include "Widgets/DeclarativeSyntaxSupport.h"
+#include "Widgets/SCompoundWidget.h"
+
+class AMediaPlate;
+
+/**
+ * Implements the details panel of the MediaPlate asset editor.
+ */
+class SMediaPlateEditorDetails
+	: public SCompoundWidget
+{
+public:
+
+	SLATE_BEGIN_ARGS(SMediaPlateEditorDetails) { }
+	SLATE_END_ARGS()
+
+public:
+
+	/**
+	 * Construct this widget.
+	 *
+	 * @param InArgs		The declaration data for this widget.
+	 * @param InMediaPlate	The MediaPlate to show the details for.
+	 * @param InStyleSet	The style set to use.
+	 */
+	void Construct(const FArguments& InArgs, AMediaPlate& InMediaPlate, const TSharedRef<ISlateStyle>& InStyle);
+
+private:
+
+	/** Pointer to the MediaPlate that is being viewed. */
+	TObjectPtr<AMediaPlate> MediaPlate;
+};
