MediaPlate: Media source does not show if the Url is being used.

#rnx

diff --git a/Engine/Plugins/Media/MediaPlate/Source/MediaPlateEditor/Private/MediaPlateCustomization.cpp b/Engine/Plugins/Media/MediaPlate/Source/MediaPlateEditor/Private/MediaPlateCustomization.cpp
--- a/Engine/Plugins/Media/MediaPlate/Source/MediaPlateEditor/Private/MediaPlateCustomization.cpp
+++ b/Engine/Plugins/Media/MediaPlate/Source/MediaPlateEditor/Private/MediaPlateCustomization.cpp
@@ -34,6 +34,37 @@ void FMediaPlateCustomization::CustomizeDetails(IDetailLayoutBuilder& DetailBuil
 		}
 	}
 
+	// Get media source property.
+	MediaSourceProperty = DetailBuilder.GetProperty(GET_MEMBER_NAME_CHECKED(AMediaPlate, MediaSource));
+	
+	// Get Url property.
+	TSharedRef<IPropertyHandle> Property = DetailBuilder.GetProperty(GET_MEMBER_NAME_CHECKED(AMediaPlate, Url));
+	if (Property->IsValidHandle())
+	{
+		UrlProperty = Property->GetChildHandle("FilePath");
+		if (UrlProperty->IsValidHandle())
+		{
+			if (MediaSourceProperty->IsValidHandle())
+			{
+				// Get a callback when this changes so we can hide the media source property.
+				FSimpleDelegate OnUrlChangedDelegate = FSimpleDelegate::CreateSP(this,
+					&FMediaPlateCustomization::OnUrlChanged, &DetailBuilder);
+				UrlProperty->SetOnPropertyValueChanged(OnUrlChangedDelegate);
+
+				// Is the Url being used?
+				FString Url;
+				if (UrlProperty->GetValue(Url))
+				{
+					if (Url.IsEmpty() == false)
+					{
+						// Yes, so hide the media source.
+						MediaSourceProperty->MarkHiddenByCustomization();
+					}
+				}
+			}
+		}
+	}
+
 	// Add media control buttons.
 	MediaPlateCategory.AddCustomRow(LOCTEXT("MediaPlateControls", "MediaPlate Controls"))
 		.NameContent()
@@ -149,4 +180,10 @@ FReply FMediaPlateCustomization::OnOpenMediaPlate()
 	return FReply::Handled();
 }
 
+void FMediaPlateCustomization::OnUrlChanged(IDetailLayoutBuilder* DetailBuilder)
+{
+	// Refresh the layout so we can show/hide the media source.
+	DetailBuilder->ForceRefreshDetails();
+}
+
 #undef LOCTEXT_NAMESPACE
diff --git a/Engine/Plugins/Media/MediaPlate/Source/MediaPlateEditor/Private/MediaPlateCustomization.h b/Engine/Plugins/Media/MediaPlate/Source/MediaPlateEditor/Private/MediaPlateCustomization.h
--- a/Engine/Plugins/Media/MediaPlate/Source/MediaPlateEditor/Private/MediaPlateCustomization.h
+++ b/Engine/Plugins/Media/MediaPlate/Source/MediaPlateEditor/Private/MediaPlateCustomization.h
@@ -37,10 +37,20 @@ private:
 	/** List of the actors we are editing. */
 	TArray<TWeakObjectPtr<AMediaPlate>> ActorsList;
 
+	/** Stores our media source property. */
+	TSharedPtr<IPropertyHandle> MediaSourceProperty;
+	/** Stores our Url property. */
+	TSharedPtr<IPropertyHandle> UrlProperty;
+
 	/**
 	 * Called when the open media plate button is pressed.
 	 */
 	FReply OnOpenMediaPlate();
 
+	/**
+	 * Called when Url changes.
+	 */
+	void OnUrlChanged(IDetailLayoutBuilder* DetailBuilder);
+
 };
 
