MediaPlate: Renamed Url to MediaPath.

#rnx

diff --git a/Engine/Plugins/Media/MediaPlate/Source/MediaPlate/Private/MediaPlate.cpp b/Engine/Plugins/Media/MediaPlate/Source/MediaPlate/Private/MediaPlate.cpp
--- a/Engine/Plugins/Media/MediaPlate/Source/MediaPlate/Private/MediaPlate.cpp
+++ b/Engine/Plugins/Media/MediaPlate/Source/MediaPlate/Private/MediaPlate.cpp
@@ -136,18 +136,18 @@ void AMediaPlate::Play()
 		MediaPlayer->PlayOnOpen = true;
 		MediaPlayer->SetLooping(bLoop);
 		
-		// Try and play the URL.
-		if (Url.FilePath.IsEmpty() == false)
+		// Try and play the media path.
+		if (MediaPath.FilePath.IsEmpty() == false)
 		{
 			// Is it a URL or a file?
-			bool bIsUrl = Url.FilePath.Contains(TEXT("://"));
+			bool bIsUrl = MediaPath.FilePath.Contains(TEXT("://"));
 			if (bIsUrl)
 			{
-				bIsPlaying = MediaPlayer->OpenUrl(Url.FilePath);
+				bIsPlaying = MediaPlayer->OpenUrl(MediaPath.FilePath);
 			}
 			else
 			{
-				bIsPlaying = MediaPlayer->OpenFile(Url.FilePath);
+				bIsPlaying = MediaPlayer->OpenFile(MediaPath.FilePath);
 			}
 		}
 
diff --git a/Engine/Plugins/Media/MediaPlate/Source/MediaPlate/Public/MediaPlate.h b/Engine/Plugins/Media/MediaPlate/Source/MediaPlate/Public/MediaPlate.h
--- a/Engine/Plugins/Media/MediaPlate/Source/MediaPlate/Public/MediaPlate.h
+++ b/Engine/Plugins/Media/MediaPlate/Source/MediaPlate/Public/MediaPlate.h
@@ -68,13 +68,13 @@ public:
 	UPROPERTY(Category = MediaPlate, VisibleAnywhere, BlueprintReadOnly)
 	TObjectPtr<UStaticMeshComponent> StaticMeshComponent;
 
-	/** What media to play. This will only be used if Url is not valid. */
+	/** URL (or file)  to play. This will take precedence over the MediaSource. */
 	UPROPERTY(EditAnywhere, Category = MediaPlate)
-	TObjectPtr<UMediaSource> MediaSource;
+	FFilePath MediaPath;
 
-	/** URL (or file)  to play. This will take precedence over the MediaSource. */
+	/** What media to play. This will only be used if MediaPath is not valid. */
 	UPROPERTY(EditAnywhere, Category = MediaPlate)
-	FFilePath Url;
+	TObjectPtr<UMediaSource> MediaSource;
 
 private:
 	/** Name for our media component. */
diff --git a/Engine/Plugins/Media/MediaPlate/Source/MediaPlateEditor/Private/MediaPlateCustomization.cpp b/Engine/Plugins/Media/MediaPlate/Source/MediaPlateEditor/Private/MediaPlateCustomization.cpp
--- a/Engine/Plugins/Media/MediaPlate/Source/MediaPlateEditor/Private/MediaPlateCustomization.cpp
+++ b/Engine/Plugins/Media/MediaPlate/Source/MediaPlateEditor/Private/MediaPlateCustomization.cpp
@@ -37,25 +37,25 @@ void FMediaPlateCustomization::CustomizeDetails(IDetailLayoutBuilder& DetailBuil
 	// Get media source property.
 	MediaSourceProperty = DetailBuilder.GetProperty(GET_MEMBER_NAME_CHECKED(AMediaPlate, MediaSource));
 	
-	// Get Url property.
-	TSharedRef<IPropertyHandle> Property = DetailBuilder.GetProperty(GET_MEMBER_NAME_CHECKED(AMediaPlate, Url));
+	// Get media path property.
+	TSharedRef<IPropertyHandle> Property = DetailBuilder.GetProperty(GET_MEMBER_NAME_CHECKED(AMediaPlate, MediaPath));
 	if (Property->IsValidHandle())
 	{
-		UrlProperty = Property->GetChildHandle("FilePath");
-		if (UrlProperty->IsValidHandle())
+		MediaPathProperty = Property->GetChildHandle("FilePath");
+		if (MediaPathProperty->IsValidHandle())
 		{
 			if (MediaSourceProperty->IsValidHandle())
 			{
 				// Get a callback when this changes so we can hide the media source property.
 				FSimpleDelegate OnUrlChangedDelegate = FSimpleDelegate::CreateSP(this,
-					&FMediaPlateCustomization::OnUrlChanged, &DetailBuilder);
-				UrlProperty->SetOnPropertyValueChanged(OnUrlChangedDelegate);
+					&FMediaPlateCustomization::OnMediaPathChanged, &DetailBuilder);
+				MediaPathProperty->SetOnPropertyValueChanged(OnUrlChangedDelegate);
 
-				// Is the Url being used?
-				FString Url;
-				if (UrlProperty->GetValue(Url))
+				// Is the media path being used?
+				FString MediaPath;
+				if (MediaPathProperty->GetValue(MediaPath))
 				{
-					if (Url.IsEmpty() == false)
+					if (MediaPath.IsEmpty() == false)
 					{
 						// Yes, so hide the media source.
 						MediaSourceProperty->MarkHiddenByCustomization();
@@ -180,7 +180,7 @@ FReply FMediaPlateCustomization::OnOpenMediaPlate()
 	return FReply::Handled();
 }
 
-void FMediaPlateCustomization::OnUrlChanged(IDetailLayoutBuilder* DetailBuilder)
+void FMediaPlateCustomization::OnMediaPathChanged(IDetailLayoutBuilder* DetailBuilder)
 {
 	// Refresh the layout so we can show/hide the media source.
 	DetailBuilder->ForceRefreshDetails();
diff --git a/Engine/Plugins/Media/MediaPlate/Source/MediaPlateEditor/Private/MediaPlateCustomization.h b/Engine/Plugins/Media/MediaPlate/Source/MediaPlateEditor/Private/MediaPlateCustomization.h
--- a/Engine/Plugins/Media/MediaPlate/Source/MediaPlateEditor/Private/MediaPlateCustomization.h
+++ b/Engine/Plugins/Media/MediaPlate/Source/MediaPlateEditor/Private/MediaPlateCustomization.h
@@ -39,8 +39,8 @@ private:
 
 	/** Stores our media source property. */
 	TSharedPtr<IPropertyHandle> MediaSourceProperty;
-	/** Stores our Url property. */
-	TSharedPtr<IPropertyHandle> UrlProperty;
+	/** Stores our media path property. */
+	TSharedPtr<IPropertyHandle> MediaPathProperty;
 
 	/**
 	 * Called when the open media plate button is pressed.
@@ -48,9 +48,9 @@ private:
 	FReply OnOpenMediaPlate();
 
 	/**
-	 * Called when Url changes.
+	 * Called when MediaPath changes.
 	 */
-	void OnUrlChanged(IDetailLayoutBuilder* DetailBuilder);
+	void OnMediaPathChanged(IDetailLayoutBuilder* DetailBuilder);
 
 };
 
