MediaPlate: Fixed playing in the media plate window only playing the media source.
#rnx

diff --git a/Engine/Plugins/Media/MediaPlate/Source/MediaPlateEditor/Private/Toolkits/MediaPlateEditorToolkit.cpp b/Engine/Plugins/Media/MediaPlate/Source/MediaPlateEditor/Private/Toolkits/MediaPlateEditorToolkit.cpp
--- a/Engine/Plugins/Media/MediaPlate/Source/MediaPlateEditor/Private/Toolkits/MediaPlateEditorToolkit.cpp
+++ b/Engine/Plugins/Media/MediaPlate/Source/MediaPlateEditor/Private/Toolkits/MediaPlateEditorToolkit.cpp
@@ -228,7 +228,7 @@ void FMediaPlateEditorToolkit::BindCommands()
 
 	ToolkitCommands->MapAction(
 		Commands.OpenMedia,
-		FExecuteAction::CreateLambda([this] { MediaPlate->GetMediaPlayer()->OpenSource(MediaPlate->MediaSource); }),
+		FExecuteAction::CreateLambda([this] { MediaPlate->Play(); }),
 		FCanExecuteAction::CreateLambda([this] { return true; })
 	);
 
