MediaPlate: Fixed media path not playing exrs.
#rnx

diff --git a/Engine/Plugins/Media/MediaPlate/Source/MediaPlate/Private/MediaPlate.cpp b/Engine/Plugins/Media/MediaPlate/Source/MediaPlate/Private/MediaPlate.cpp
--- a/Engine/Plugins/Media/MediaPlate/Source/MediaPlate/Private/MediaPlate.cpp
+++ b/Engine/Plugins/Media/MediaPlate/Source/MediaPlate/Private/MediaPlate.cpp
@@ -8,6 +8,7 @@
 #include "MediaComponent.h"
 #include "MediaPlayer.h"
 #include "MediaPlateModule.h"
+#include "MediaSource.h"
 #include "MediaTexture.h"
 #include "MediaTextureTracker.h"
 #include "UObject/ConstructorHelpers.h"
@@ -139,15 +140,11 @@ void AMediaPlate::Play()
 		// Try and play the media path.
 		if (MediaPath.FilePath.IsEmpty() == false)
 		{
-			// Is it a URL or a file?
-			bool bIsUrl = MediaPath.FilePath.Contains(TEXT("://"));
-			if (bIsUrl)
+			MediaPathMediaSource = UMediaSource::SpawnMediaSourceForString(MediaPath.FilePath);
+			if (MediaPathMediaSource != nullptr)
 			{
-				bIsPlaying = MediaPlayer->OpenUrl(MediaPath.FilePath);
-			}
-			else
-			{
-				bIsPlaying = MediaPlayer->OpenFile(MediaPath.FilePath);
+				bIsPlaying = MediaPlayer->OpenSource(MediaPathMediaSource);
+				return;
 			}
 		}
 
diff --git a/Engine/Plugins/Media/MediaPlate/Source/MediaPlate/Public/MediaPlate.h b/Engine/Plugins/Media/MediaPlate/Source/MediaPlate/Public/MediaPlate.h
--- a/Engine/Plugins/Media/MediaPlate/Source/MediaPlate/Public/MediaPlate.h
+++ b/Engine/Plugins/Media/MediaPlate/Source/MediaPlate/Public/MediaPlate.h
@@ -85,6 +85,10 @@ private:
 	/** Info representing this object. */
 	TSharedPtr<FMediaTextureTrackerObject, ESPMode::ThreadSafe> MediaTextureTrackerObject;
 
+	/** If we are using MediaPath, then this is the media source for it. */
+	UPROPERTY(Transient)
+	TObjectPtr<UMediaSource> MediaPathMediaSource;
+
 	/**
 	 * Adds our media texture to the media texture tracker.
 	 */
