Branching

//UE5/Main/Engine/Plugins/Media/MediaPlate/Source/MediaPlate/Public/MediaPlate.h

to //UE5/Main/Engine/Plugins/Media/MediaPlate/Source/MediaPlate/Public/MediaPlateComponent.h

diff --git a/Engine/Plugins/Media/MediaPlate/Source/MediaPlate/Public/MediaPlateComponent.h b/Engine/Plugins/Media/MediaPlate/Source/MediaPlate/Public/MediaPlateComponent.h
new file mode 100644
--- /dev/null
+++ b/Engine/Plugins/Media/MediaPlate/Source/MediaPlate/Public/MediaPlateComponent.h
@@ -0,0 +1,100 @@
+// Copyright Epic Games, Inc. All Rights Reserved.
+
+#pragma once
+
+#include "CoreMinimal.h"
+#include "GameFramework/Actor.h"
+#include "MediaPlate.generated.h"
+
+class UMediaComponent;
+class UMediaPlayer;
+class UMediaSource;
+class UMediaTexture;
+struct FMediaTextureTrackerObject;
+
+/**
+ * MediaPlate is an actor that can play and show media in the world.
+ */
+UCLASS()
+class MEDIAPLATE_API AMediaPlate : public AActor
+{
+	GENERATED_UCLASS_BODY()
+
+public:
+	//~ Begin AActor Interface
+	virtual void PostRegisterAllComponents() override;
+	virtual void BeginPlay() override;
+	virtual void BeginDestroy() override;
+	//~ End AActor Interface
+
+	/**
+	 * Call this get our media player.
+	 */
+	UFUNCTION(BlueprintCallable, Category = "Media|MediaPlate")
+	UMediaPlayer* GetMediaPlayer();
+
+	/**
+	 * Call this get our media texture.
+	 */
+	UFUNCTION(BlueprintCallable, Category = "Media|MediaPlate")
+	UMediaTexture* GetMediaTexture();
+
+	/**
+	 * Call this to start playing.
+	 */
+	UFUNCTION(BlueprintCallable, Category = "Media|MediaPlate")
+	void Play();
+
+	/**
+	 * Call this to stop playing.
+	 */
+	UFUNCTION(BlueprintCallable, Category = "Media|MediaPlate")
+	void Stop();
+
+	/** If set then start playing right away. */
+	UPROPERTY(EditAnywhere, Category = "MediaPlate")
+	bool bAutoPlay;
+
+	/** If set then loop when we reach the end. */
+	UPROPERTY(EditAnywhere, Category = "MediaPlate")
+	bool bLoop;
+
+	/** Holds the media player. */
+	UPROPERTY(Category = MediaPlate, VisibleAnywhere, BlueprintReadOnly)
+	TObjectPtr<UMediaComponent> MediaComponent;
+
+
+	/** Holds the mesh. */
+	UPROPERTY(Category = MediaPlate, VisibleAnywhere, BlueprintReadOnly)
+	TObjectPtr<UStaticMeshComponent> StaticMeshComponent;
+
+	/** URL (or file)  to play. This will take precedence over the MediaSource. */
+	UPROPERTY(EditAnywhere, Category = MediaPlate)
+	FFilePath MediaPath;
+
+	/** What media to play. This will only be used if MediaPath is not valid. */
+	UPROPERTY(EditAnywhere, Category = MediaPlate)
+	TObjectPtr<UMediaSource> MediaSource;
+
+private:
+	/** Name for our media component. */
+	static FLazyName MediaComponentName;
+	/** Name for the media texture parameter in the material. */
+	static FLazyName MediaTextureName;
+
+	/** Info representing this object. */
+	TSharedPtr<FMediaTextureTrackerObject, ESPMode::ThreadSafe> MediaTextureTrackerObject;
+
+	/** If we are using MediaPath, then this is the media source for it. */
+	UPROPERTY(Transient)
+	TObjectPtr<UMediaSource> MediaPathMediaSource;
+
+	/**
+	 * Adds our media texture to the media texture tracker.
+	 */
+	void RegisterWithMediaTextureTracker();
+	/**
+	 * Removes our texture from the media texture tracker.
+	 */
+	void UnregisterWithMediaTextureTracker();
+};
