MediaPlate: Added media plate component.

#rnx

diff --git a/Engine/Plugins/Media/MediaPlate/Source/MediaPlate/Private/MediaPlate.cpp b/Engine/Plugins/Media/MediaPlate/Source/MediaPlate/Private/MediaPlate.cpp
--- a/Engine/Plugins/Media/MediaPlate/Source/MediaPlate/Private/MediaPlate.cpp
+++ b/Engine/Plugins/Media/MediaPlate/Source/MediaPlate/Private/MediaPlate.cpp
@@ -5,17 +5,14 @@
 #include "Components/StaticMeshComponent.h"
 #include "Engine/StaticMesh.h"
 #include "Materials/MaterialInstanceDynamic.h"
-#include "MediaComponent.h"
-#include "MediaPlayer.h"
+#include "MediaPlateComponent.h"
 #include "MediaPlateModule.h"
-#include "MediaSource.h"
 #include "MediaTexture.h"
-#include "MediaTextureTracker.h"
 #include "UObject/ConstructorHelpers.h"
 
 #define LOCTEXT_NAMESPACE "MediaPlate"
 
-FLazyName AMediaPlate::MediaComponentName(TEXT("MediaComponent0"));
+FLazyName AMediaPlate::MediaPlateComponentName(TEXT("MediaPlateComponent0"));
 FLazyName AMediaPlate::MediaTextureName("MediaTexture");
 
 AMediaPlate::AMediaPlate(const FObjectInitializer& ObjectInitializer)
@@ -24,16 +21,7 @@ AMediaPlate::AMediaPlate(const FObjectInitializer& ObjectInitializer)
 	RootComponent = CreateDefaultSubobject<USceneComponent>(TEXT("Root"));
 
 	// Set up media component.
-	MediaComponent = CreateDefaultSubobject<UMediaComponent>(MediaComponentName);
-	if (MediaComponent != nullptr)
-	{
-		// Set up media texture.
-		UMediaTexture* MediaTexture = MediaComponent->GetMediaTexture();
-		if (MediaTexture != nullptr)
-		{
-			MediaTexture->NewStyleOutput = true;
-		}
-	}
+	MediaPlateComponent = CreateDefaultSubobject<UMediaPlateComponent>(MediaPlateComponentName);
 
 	// Set up static mesh component.
 	StaticMeshComponent = CreateDefaultSubobject<UStaticMeshComponent>(TEXT("StaticMeshComponent"));
@@ -60,142 +48,46 @@ void AMediaPlate::PostRegisterAllComponents()
 {
 	Super::PostRegisterAllComponents();
 
-	if (StaticMeshComponent != nullptr)
+	if (MediaPlateComponent != nullptr)
 	{
-		// Do we have a material?
-		if (StaticMeshComponent->GetNumOverrideMaterials() == 0)
+		if (StaticMeshComponent != nullptr)
 		{
-			// Add material.
-			UMaterial* Material = LoadObject<UMaterial>(nullptr, TEXT("/MediaPlate/M_MediaPlate"));
-			if (Material != nullptr)
+			// Do we have a material?
+			if (StaticMeshComponent->GetNumOverrideMaterials() == 0)
 			{
-				UMaterialInstanceDynamic* MaterialInstance = StaticMeshComponent->CreateAndSetMaterialInstanceDynamicFromMaterial(0, Material);
+				// Add material.
+				UMaterial* Material = LoadObject<UMaterial>(nullptr, TEXT("/MediaPlate/M_MediaPlate"));
+				if (Material != nullptr)
+				{
+					UMaterialInstanceDynamic* MaterialInstance = StaticMeshComponent->CreateAndSetMaterialInstanceDynamicFromMaterial(0, Material);
+				}
 			}
-		}
 
-		// Set up the material to point to our media texture.
-		if (StaticMeshComponent->GetNumMaterials() > 0)
-		{
-			UMaterialInstanceDynamic* MaterialInstance = Cast< UMaterialInstanceDynamic>(StaticMeshComponent->GetMaterial(0));
-			if (MaterialInstance != nullptr)
+			// Set up the material to point to our media texture.
+			if (StaticMeshComponent->GetNumMaterials() > 0)
 			{
-				MaterialInstance->SetTextureParameterValue(MediaTextureName, MediaComponent->GetMediaTexture());
+				UMaterialInstanceDynamic* MaterialInstance = Cast< UMaterialInstanceDynamic>(StaticMeshComponent->GetMaterial(0));
+				if (MaterialInstance != nullptr)
+				{
+					MaterialInstance->SetTextureParameterValue(MediaTextureName, MediaPlateComponent->GetMediaTexture());
+				}
 			}
 		}
-	}
 
-	// Add our media texture to the tracker.
-	RegisterWithMediaTextureTracker();
-}
-
-void AMediaPlate::BeginPlay()
-{
-	Super::BeginPlay();
-
-	if (bAutoPlay)
-	{
-		Play();
+		// Add our media texture to the tracker.
+		MediaPlateComponent->RegisterWithMediaTextureTracker();
 	}
 }
 
 void AMediaPlate::BeginDestroy()
 {
-	// Remove our media texture.
-	UnregisterWithMediaTextureTracker();
-	
-	Super::BeginDestroy();
-}
-
-UMediaPlayer* AMediaPlate::GetMediaPlayer()
-{
-	TObjectPtr<UMediaPlayer> MediaPlayer = nullptr;
-	if (MediaComponent != nullptr)
-	{
-		MediaPlayer = MediaComponent->GetMediaPlayer();
-	}
-
-	return MediaPlayer;
-}
-
-UMediaTexture* AMediaPlate::GetMediaTexture()
-{
-	TObjectPtr<UMediaTexture> MediaTexture = nullptr;
-	if (MediaComponent != nullptr)
+	if (MediaPlateComponent != nullptr)
 	{
-		MediaTexture = MediaComponent->GetMediaTexture();
-	}
-
-	return MediaTexture;
-}
-
-void AMediaPlate::Play()
-{
-	bool bIsPlaying = false;
-	TObjectPtr<UMediaPlayer> MediaPlayer = GetMediaPlayer();
-	if (MediaPlayer != nullptr)
-	{
-		MediaPlayer->PlayOnOpen = true;
-		MediaPlayer->SetLooping(bLoop);
-		
-		// Try and play the media path.
-		if (MediaPath.FilePath.IsEmpty() == false)
-		{
-			MediaPathMediaSource = UMediaSource::SpawnMediaSourceForString(MediaPath.FilePath);
-			if (MediaPathMediaSource != nullptr)
-			{
-				bIsPlaying = MediaPlayer->OpenSource(MediaPathMediaSource);
-				return;
-			}
-		}
-
-		// If we did not get anything, try the media source.
-		if (bIsPlaying == false)
-		{
-			bIsPlaying = MediaPlayer->OpenSource(MediaSource);
-		}
-	}
-
-	// Did anything play?
-	if (bIsPlaying == false)
-	{
-		UE_LOG(LogMediaPlate, Warning, TEXT("Could not play anything."));
-	}
-}
-
-void AMediaPlate::Stop()
-{
-	TObjectPtr<UMediaPlayer> MediaPlayer = GetMediaPlayer();
-	if (MediaPlayer != nullptr)
-	{
-		MediaPlayer->Close();
-	}
-}
-
-void AMediaPlate::RegisterWithMediaTextureTracker()
-{
-	// Set up object.
-	MediaTextureTrackerObject = MakeShared<FMediaTextureTrackerObject, ESPMode::ThreadSafe>();
-	MediaTextureTrackerObject->Object = this;
-	MediaTextureTrackerObject->MipMapLODBias = 0.0f;
-
-	// Add our texture.
-	TObjectPtr<UMediaTexture> MediaTexture = GetMediaTexture();
-	if (MediaTexture != nullptr)
-	{
-		FMediaTextureTracker& MediaTextureTracker = FMediaTextureTracker::Get();
-		MediaTextureTracker.RegisterTexture(MediaTextureTrackerObject, MediaTexture);
-	}
-}
-
-void AMediaPlate::UnregisterWithMediaTextureTracker()
-{
-	// Remove out texture.
-	if (MediaTextureTrackerObject != nullptr)
-	{
-		FMediaTextureTracker& MediaTextureTracker = FMediaTextureTracker::Get();
-		TObjectPtr<UMediaTexture> MediaTexture = GetMediaTexture();
-		MediaTextureTracker.UnregisterTexture(MediaTextureTrackerObject, MediaTexture);
+		// Remove our media texture.
+		MediaPlateComponent->UnregisterWithMediaTextureTracker();
 	}
+	
+	Super::BeginDestroy();
 }
 
 #undef LOCTEXT_NAMESPACE
diff --git a/Engine/Plugins/Media/MediaPlate/Source/MediaPlate/Private/MediaPlateComponent.cpp b/Engine/Plugins/Media/MediaPlate/Source/MediaPlate/Private/MediaPlateComponent.cpp
new file mode 100644
--- /dev/null
+++ b/Engine/Plugins/Media/MediaPlate/Source/MediaPlate/Private/MediaPlateComponent.cpp
@@ -0,0 +1,137 @@
+// Copyright Epic Games, Inc. All Rights Reserved.
+
+#include "MediaPlateComponent.h"
+
+#include "Components/StaticMeshComponent.h"
+#include "Engine/StaticMesh.h"
+#include "MediaComponent.h"
+#include "MediaPlayer.h"
+#include "MediaPlateModule.h"
+#include "MediaSource.h"
+#include "MediaTexture.h"
+#include "MediaTextureTracker.h"
+
+#define LOCTEXT_NAMESPACE "MediaPlate"
+
+FLazyName UMediaPlateComponent::MediaComponentName(TEXT("MediaComponent0"));
+
+UMediaPlateComponent::UMediaPlateComponent(const FObjectInitializer& ObjectInitializer)
+	: Super(ObjectInitializer)
+{
+	// Set up media component.
+	MediaComponent = CreateDefaultSubobject<UMediaComponent>(MediaComponentName);
+	if (MediaComponent != nullptr)
+	{
+		// Set up media texture.
+		UMediaTexture* MediaTexture = MediaComponent->GetMediaTexture();
+		if (MediaTexture != nullptr)
+		{
+			MediaTexture->NewStyleOutput = true;
+		}
+	}
+}
+
+void UMediaPlateComponent::BeginPlay()
+{
+	Super::BeginPlay();
+
+	// Start playing?
+	if (bAutoPlay)
+	{
+		Play();
+	}
+}
+
+UMediaPlayer* UMediaPlateComponent::GetMediaPlayer()
+{
+	TObjectPtr<UMediaPlayer> MediaPlayer = nullptr;
+	if (MediaComponent != nullptr)
+	{
+		MediaPlayer = MediaComponent->GetMediaPlayer();
+	}
+
+	return MediaPlayer;
+}
+
+UMediaTexture* UMediaPlateComponent::GetMediaTexture()
+{
+	TObjectPtr<UMediaTexture> MediaTexture = nullptr;
+	if (MediaComponent != nullptr)
+	{
+		MediaTexture = MediaComponent->GetMediaTexture();
+	}
+
+	return MediaTexture;
+}
+
+void UMediaPlateComponent::Play()
+{
+	bool bIsPlaying = false;
+	TObjectPtr<UMediaPlayer> MediaPlayer = GetMediaPlayer();
+	if (MediaPlayer != nullptr)
+	{
+		MediaPlayer->PlayOnOpen = true;
+		MediaPlayer->SetLooping(bLoop);
+		
+		// Try and play the media path.
+		if (MediaPath.FilePath.IsEmpty() == false)
+		{
+			MediaPathMediaSource = UMediaSource::SpawnMediaSourceForString(MediaPath.FilePath);
+			if (MediaPathMediaSource != nullptr)
+			{
+				bIsPlaying = MediaPlayer->OpenSource(MediaPathMediaSource);
+				return;
+			}
+		}
+
+		// If we did not get anything, try the media source.
+		if (bIsPlaying == false)
+		{
+			bIsPlaying = MediaPlayer->OpenSource(MediaSource);
+		}
+	}
+
+	// Did anything play?
+	if (bIsPlaying == false)
+	{
+		UE_LOG(LogMediaPlate, Warning, TEXT("Could not play anything."));
+	}
+}
+
+void UMediaPlateComponent::Stop()
+{
+	TObjectPtr<UMediaPlayer> MediaPlayer = GetMediaPlayer();
+	if (MediaPlayer != nullptr)
+	{
+		MediaPlayer->Close();
+	}
+}
+
+void UMediaPlateComponent::RegisterWithMediaTextureTracker()
+{
+	// Set up object.
+	MediaTextureTrackerObject = MakeShared<FMediaTextureTrackerObject, ESPMode::ThreadSafe>();
+	MediaTextureTrackerObject->Object = GetOwner();
+	MediaTextureTrackerObject->MipMapLODBias = 0.0f;
+
+	// Add our texture.
+	TObjectPtr<UMediaTexture> MediaTexture = GetMediaTexture();
+	if (MediaTexture != nullptr)
+	{
+		FMediaTextureTracker& MediaTextureTracker = FMediaTextureTracker::Get();
+		MediaTextureTracker.RegisterTexture(MediaTextureTrackerObject, MediaTexture);
+	}
+}
+
+void UMediaPlateComponent::UnregisterWithMediaTextureTracker()
+{
+	// Remove out texture.
+	if (MediaTextureTrackerObject != nullptr)
+	{
+		FMediaTextureTracker& MediaTextureTracker = FMediaTextureTracker::Get();
+		TObjectPtr<UMediaTexture> MediaTexture = GetMediaTexture();
+		MediaTextureTracker.UnregisterTexture(MediaTextureTrackerObject, MediaTexture);
+	}
+}
+
+#undef LOCTEXT_NAMESPACE
diff --git a/Engine/Plugins/Media/MediaPlate/Source/MediaPlate/Public/MediaPlate.h b/Engine/Plugins/Media/MediaPlate/Source/MediaPlate/Public/MediaPlate.h
--- a/Engine/Plugins/Media/MediaPlate/Source/MediaPlate/Public/MediaPlate.h
+++ b/Engine/Plugins/Media/MediaPlate/Source/MediaPlate/Public/MediaPlate.h
@@ -6,11 +6,7 @@
 #include "GameFramework/Actor.h"
 #include "MediaPlate.generated.h"
 
-class UMediaComponent;
-class UMediaPlayer;
-class UMediaSource;
-class UMediaTexture;
-struct FMediaTextureTrackerObject;
+class UMediaPlateComponent;
 
 /**
  * MediaPlate is an actor that can play and show media in the world.
@@ -23,78 +19,19 @@ class MEDIAPLATE_API AMediaPlate : public AActor
 public:
 	//~ Begin AActor Interface
 	virtual void PostRegisterAllComponents() override;
-	virtual void BeginPlay() override;
 	virtual void BeginDestroy() override;
 	//~ End AActor Interface
 
-	/**
-	 * Call this get our media player.
-	 */
-	UFUNCTION(BlueprintCallable, Category = "Media|MediaPlate")
-	UMediaPlayer* GetMediaPlayer();
-
-	/**
-	 * Call this get our media texture.
-	 */
-	UFUNCTION(BlueprintCallable, Category = "Media|MediaPlate")
-	UMediaTexture* GetMediaTexture();
-
-	/**
-	 * Call this to start playing.
-	 */
-	UFUNCTION(BlueprintCallable, Category = "Media|MediaPlate")
-	void Play();
-
-	/**
-	 * Call this to stop playing.
-	 */
-	UFUNCTION(BlueprintCallable, Category = "Media|MediaPlate")
-	void Stop();
-
-	/** If set then start playing right away. */
-	UPROPERTY(EditAnywhere, Category = "MediaPlate")
-	bool bAutoPlay;
-
-	/** If set then loop when we reach the end. */
-	UPROPERTY(EditAnywhere, Category = "MediaPlate")
-	bool bLoop;
-
-	/** Holds the media player. */
 	UPROPERTY(Category = MediaPlate, VisibleAnywhere, BlueprintReadOnly)
-	TObjectPtr<UMediaComponent> MediaComponent;
-
+	TObjectPtr<UMediaPlateComponent> MediaPlateComponent;
 
 	/** Holds the mesh. */
 	UPROPERTY(Category = MediaPlate, VisibleAnywhere, BlueprintReadOnly)
 	TObjectPtr<UStaticMeshComponent> StaticMeshComponent;
 
-	/** URL (or file)  to play. This will take precedence over the MediaSource. */
-	UPROPERTY(EditAnywhere, Category = MediaPlate)
-	FFilePath MediaPath;
-
-	/** What media to play. This will only be used if MediaPath is not valid. */
-	UPROPERTY(EditAnywhere, Category = MediaPlate)
-	TObjectPtr<UMediaSource> MediaSource;
-
 private:
-	/** Name for our media component. */
-	static FLazyName MediaComponentName;
+	/** Name for our media plate component. */
+	static FLazyName MediaPlateComponentName;
 	/** Name for the media texture parameter in the material. */
 	static FLazyName MediaTextureName;
-
-	/** Info representing this object. */
-	TSharedPtr<FMediaTextureTrackerObject, ESPMode::ThreadSafe> MediaTextureTrackerObject;
-
-	/** If we are using MediaPath, then this is the media source for it. */
-	UPROPERTY(Transient)
-	TObjectPtr<UMediaSource> MediaPathMediaSource;
-
-	/**
-	 * Adds our media texture to the media texture tracker.
-	 */
-	void RegisterWithMediaTextureTracker();
-	/**
-	 * Removes our texture from the media texture tracker.
-	 */
-	void UnregisterWithMediaTextureTracker();
 };
diff --git a/Engine/Plugins/Media/MediaPlate/Source/MediaPlate/Public/MediaPlateComponent.h b/Engine/Plugins/Media/MediaPlate/Source/MediaPlate/Public/MediaPlateComponent.h
new file mode 100644
--- /dev/null
+++ b/Engine/Plugins/Media/MediaPlate/Source/MediaPlate/Public/MediaPlateComponent.h
@@ -0,0 +1,90 @@
+// Copyright Epic Games, Inc. All Rights Reserved.
+
+#pragma once
+
+#include "Components/ActorComponent.h"
+#include "CoreMinimal.h"
+#include "MediaPlateComponent.generated.h"
+
+class UMediaComponent;
+class UMediaPlayer;
+class UMediaSource;
+class UMediaTexture;
+struct FMediaTextureTrackerObject;
+
+/**
+ * This is a component for AMediaPlate that can play and show media in the world.
+ */
+UCLASS()
+class MEDIAPLATE_API UMediaPlateComponent : public UActorComponent
+{
+	GENERATED_UCLASS_BODY()
+
+public:
+	//~ UActorComponent interface.
+	virtual void BeginPlay();
+
+	/**
+	 * Call this get our media player.
+	 */
+	UFUNCTION(BlueprintCallable, Category = "Media|MediaPlate")
+	UMediaPlayer* GetMediaPlayer();
+
+	/**
+	 * Call this get our media texture.
+	 */
+	UFUNCTION(BlueprintCallable, Category = "Media|MediaPlate")
+	UMediaTexture* GetMediaTexture();
+
+	/**
+	 * Call this to start playing.
+	 */
+	UFUNCTION(BlueprintCallable, Category = "Media|MediaPlate")
+	void Play();
+
+	/**
+	 * Call this to stop playing.
+	 */
+	UFUNCTION(BlueprintCallable, Category = "Media|MediaPlate")
+	void Stop();
+
+	/** If set then start playing right away. */
+	UPROPERTY(EditAnywhere, Category = "MediaPlate")
+	bool bAutoPlay;
+
+	/** If set then loop when we reach the end. */
+	UPROPERTY(EditAnywhere, Category = "MediaPlate")
+	bool bLoop;
+
+	/** Holds the media player. */
+	UPROPERTY(Category = MediaPlate, VisibleAnywhere, BlueprintReadOnly)
+	TObjectPtr<UMediaComponent> MediaComponent;
+
+	/** URL (or file)  to play. This will take precedence over the MediaSource. */
+	UPROPERTY(EditAnywhere, Category = MediaPlate)
+	FFilePath MediaPath;
+
+	/** What media to play. This will only be used if MediaPath is not valid. */
+	UPROPERTY(EditAnywhere, Category = MediaPlate)
+	TObjectPtr<UMediaSource> MediaSource;
+
+	/**
+	 * Adds our media texture to the media texture tracker.
+	 */
+	void RegisterWithMediaTextureTracker();
+	/**
+	 * Removes our texture from the media texture tracker.
+	 */
+	void UnregisterWithMediaTextureTracker();
+
+private:
+	/** Name for our media component. */
+	static FLazyName MediaComponentName;
+
+	/** Info representing this object. */
+	TSharedPtr<FMediaTextureTrackerObject, ESPMode::ThreadSafe> MediaTextureTrackerObject;
+
+	/** If we are using MediaPath, then this is the media source for it. */
+	UPROPERTY(Transient)
+	TObjectPtr<UMediaSource> MediaPathMediaSource;
+};
diff --git a/Engine/Plugins/Media/MediaPlate/Source/MediaPlateEditor/Private/ActorFactoryMediaPlate.cpp b/Engine/Plugins/Media/MediaPlate/Source/MediaPlateEditor/Private/ActorFactoryMediaPlate.cpp
--- a/Engine/Plugins/Media/MediaPlate/Source/MediaPlateEditor/Private/ActorFactoryMediaPlate.cpp
+++ b/Engine/Plugins/Media/MediaPlate/Source/MediaPlateEditor/Private/ActorFactoryMediaPlate.cpp
@@ -4,6 +4,7 @@
 
 #include "MediaSource.h"
 #include "MediaPlate.h"
+#include "MediaPlateComponent.h"
 
 #define LOCTEXT_NAMESPACE "ActorFactoryMediaPlate"
 
@@ -54,9 +55,9 @@ void UActorFactoryMediaPlate::SetUpActor(UObject* Asset, AActor* Actor)
 
 		// Hook up media source.
 		UMediaSource* MediaSource = Cast<UMediaSource>(Asset);
-		if (MediaSource != nullptr)
+		if ((MediaSource != nullptr) && (MediaPlate->MediaPlateComponent != nullptr))
 		{
-			MediaPlate->MediaSource = MediaSource;
+			MediaPlate->MediaPlateComponent->MediaSource = MediaSource;
 		}
 	}
 }
diff --git a/Engine/Plugins/Media/MediaPlate/Source/MediaPlateEditor/Private/AssetTools/MediaPlateActions.cpp b/Engine/Plugins/Media/MediaPlate/Source/MediaPlateEditor/Private/AssetTools/MediaPlateActions.cpp
--- a/Engine/Plugins/Media/MediaPlate/Source/MediaPlateEditor/Private/AssetTools/MediaPlateActions.cpp
+++ b/Engine/Plugins/Media/MediaPlate/Source/MediaPlateEditor/Private/AssetTools/MediaPlateActions.cpp
@@ -5,7 +5,7 @@
 #include "EditorStyleSet.h"
 #include "Framework/MultiBox/MultiBoxBuilder.h"
 #include "IContentBrowserSingleton.h"
-#include "MediaPlate.h"
+#include "MediaPlateComponent.h"
 #include "MediaTexture.h"
 #include "Misc/PackageName.h"
 #include "Toolkits/MediaPlateEditorToolkit.h"
@@ -40,7 +40,7 @@ FText FMediaPlateActions::GetName() const
 
 UClass* FMediaPlateActions::GetSupportedClass() const
 {
-	return AMediaPlate::StaticClass();
+	return UMediaPlateComponent::StaticClass();
 }
 
 FColor FMediaPlateActions::GetTypeColor() const
@@ -56,7 +56,7 @@ void FMediaPlateActions::OpenAssetEditor(const TArray<UObject*>& InObjects, TSha
 
 	for (auto ObjIt = InObjects.CreateConstIterator(); ObjIt; ++ObjIt)
 	{
-		AMediaPlate* MediaPlate = Cast<AMediaPlate>(*ObjIt);
+		UMediaPlateComponent* MediaPlate = Cast<UMediaPlateComponent>(*ObjIt);
 
 		if (MediaPlate != nullptr)
 		{
diff --git a/Engine/Plugins/Media/MediaPlate/Source/MediaPlateEditor/Private/MediaPlateCustomization.cpp b/Engine/Plugins/Media/MediaPlate/Source/MediaPlateEditor/Private/MediaPlateCustomization.cpp
--- a/Engine/Plugins/Media/MediaPlate/Source/MediaPlateEditor/Private/MediaPlateCustomization.cpp
+++ b/Engine/Plugins/Media/MediaPlate/Source/MediaPlateEditor/Private/MediaPlateCustomization.cpp
@@ -7,6 +7,7 @@
 #include "DetailWidgetRow.h"
 #include "EditorStyleSet.h"
 #include "MediaPlate.h"
+#include "MediaPlateComponent.h"
 #include "MediaPlateEditorModule.h"
 #include "Subsystems/AssetEditorSubsystem.h"
 #include "Widgets/Input/SButton.h"
@@ -24,21 +25,21 @@ void FMediaPlateCustomization::CustomizeDetails(IDetailLayoutBuilder& DetailBuil
 	// Get objects we are editing.
 	TArray<TWeakObjectPtr<UObject>> Objects;
 	DetailBuilder.GetObjectsBeingCustomized(Objects);
-	ActorsList.Reserve(Objects.Num());
+	MediaPlatesList.Reserve(Objects.Num());
 	for (TWeakObjectPtr<UObject>& Obj : Objects)
 	{
-		TWeakObjectPtr<AMediaPlate> ActorPtr = Cast<AMediaPlate>(Obj.Get());
-		if (ActorPtr.IsValid())
+		TWeakObjectPtr<UMediaPlateComponent> MediaPlate = Cast<UMediaPlateComponent>(Obj.Get());
+		if (MediaPlate.IsValid())
 		{
-			ActorsList.Add(ActorPtr);
+			MediaPlatesList.Add(MediaPlate);
 		}
 	}
 
 	// Get media source property.
-	MediaSourceProperty = DetailBuilder.GetProperty(GET_MEMBER_NAME_CHECKED(AMediaPlate, MediaSource));
+	MediaSourceProperty = DetailBuilder.GetProperty(GET_MEMBER_NAME_CHECKED(UMediaPlateComponent, MediaSource));
 	
 	// Get media path property.
-	TSharedRef<IPropertyHandle> Property = DetailBuilder.GetProperty(GET_MEMBER_NAME_CHECKED(AMediaPlate, MediaPath));
+	TSharedRef<IPropertyHandle> Property = DetailBuilder.GetProperty(GET_MEMBER_NAME_CHECKED(UMediaPlateComponent, MediaPath));
 	if (Property->IsValidHandle())
 	{
 		MediaPathProperty = Property->GetChildHandle("FilePath");
@@ -86,9 +87,9 @@ void FMediaPlateCustomization::CustomizeDetails(IDetailLayoutBuilder& DetailBuil
 						.VAlign(VAlign_Center)
 						.OnClicked_Lambda([this]() -> FReply
 						{
-							for (TWeakObjectPtr<AMediaPlate>& ActorPtr : ActorsList)
+							for (TWeakObjectPtr<UMediaPlateComponent>& MediaPlatePtr : MediaPlatesList)
 							{
-								AMediaPlate* MediaPlate = ActorPtr.Get();
+								UMediaPlateComponent* MediaPlate = MediaPlatePtr.Get();
 								if (MediaPlate != nullptr)
 								{
 									// Tell the editor module that this media plate is playing.
@@ -120,9 +121,9 @@ void FMediaPlateCustomization::CustomizeDetails(IDetailLayoutBuilder& DetailBuil
 						.VAlign(VAlign_Center)
 						.OnClicked_Lambda([this]() -> FReply
 						{
-							for (TWeakObjectPtr<AMediaPlate>& ActorPtr : ActorsList)
+							for (TWeakObjectPtr<UMediaPlateComponent>& MediaPlatePtr : MediaPlatesList)
 							{
-								AMediaPlate* MediaPlate = ActorPtr.Get();
+								UMediaPlateComponent* MediaPlate = MediaPlatePtr.Get();
 								if (MediaPlate != nullptr)
 								{
 									MediaPlate->Stop();
@@ -162,9 +163,9 @@ FReply FMediaPlateCustomization::OnOpenMediaPlate()
 {
 	// Get all our objects.
 	TArray<UObject*> AssetArray;
-	for (TWeakObjectPtr<AMediaPlate>& ActorPtr : ActorsList)
+	for (TWeakObjectPtr<UMediaPlateComponent>& MediaPlatePtr : MediaPlatesList)
 	{
-		AMediaPlate* MediaPlate = ActorPtr.Get();
+		UMediaPlateComponent* MediaPlate = MediaPlatePtr.Get();
 		if (MediaPlate != nullptr)
 		{
 			AssetArray.Add(MediaPlate);
diff --git a/Engine/Plugins/Media/MediaPlate/Source/MediaPlateEditor/Private/MediaPlateCustomization.h b/Engine/Plugins/Media/MediaPlate/Source/MediaPlateEditor/Private/MediaPlateCustomization.h
--- a/Engine/Plugins/Media/MediaPlate/Source/MediaPlateEditor/Private/MediaPlateCustomization.h
+++ b/Engine/Plugins/Media/MediaPlate/Source/MediaPlateEditor/Private/MediaPlateCustomization.h
@@ -7,11 +7,11 @@
 #include "IDetailCustomization.h"
 #include "Input/Reply.h"
 
-class AMediaPlate;
 class IDetailLayoutBuilder;
+class UMediaPlateComponent;
 
 /**
- * Implements a details view customization for the AMediaPlate class.
+ * Implements a details view customization for the UMediaPlateComponent class.
  */
 class FMediaPlateCustomization
 	: public IDetailCustomization
@@ -34,8 +34,8 @@ public:
 
 private:
 
-	/** List of the actors we are editing. */
-	TArray<TWeakObjectPtr<AMediaPlate>> ActorsList;
+	/** List of the media plates we are editing. */
+	TArray<TWeakObjectPtr<UMediaPlateComponent>> MediaPlatesList;
 
 	/** Stores our media source property. */
 	TSharedPtr<IPropertyHandle> MediaSourceProperty;
diff --git a/Engine/Plugins/Media/MediaPlate/Source/MediaPlateEditor/Private/MediaPlateEditorModule.cpp b/Engine/Plugins/Media/MediaPlate/Source/MediaPlateEditor/Private/MediaPlateEditorModule.cpp
--- a/Engine/Plugins/Media/MediaPlate/Source/MediaPlateEditor/Private/MediaPlateEditorModule.cpp
+++ b/Engine/Plugins/Media/MediaPlate/Source/MediaPlateEditor/Private/MediaPlateEditorModule.cpp
@@ -5,7 +5,7 @@
 #include "AssetTools/MediaPlateActions.h"
 #include "Editor.h"
 #include "LevelEditorViewport.h"
-#include "MediaPlate.h"
+#include "MediaPlateComponent.h"
 #include "MediaPlateCustomization.h"
 #include "MediaPlateEditorStyle.h"
 #include "MediaPlayer.h"
@@ -23,7 +23,7 @@ void FMediaPlateEditorModule::StartupModule()
 	RegisterAssetTools();
 
 	// Register customizations.
-	MediaPlateName = AMediaPlate::StaticClass()->GetFName();
+	MediaPlateName = UMediaPlateComponent::StaticClass()->GetFName();
 
 	FPropertyEditorModule& PropertyModule = FModuleManager::LoadModuleChecked<FPropertyEditorModule>("PropertyEditor");
 	PropertyModule.RegisterCustomClassLayout(MediaPlateName,
@@ -47,8 +47,8 @@ void FMediaPlateEditorModule::Tick(float DeltaTime)
 	for (int Index = 0; Index < ActiveMediaPlates.Num(); ++Index)
 	{
 		// Get the player.
-		TWeakObjectPtr<AMediaPlate>& PlatePtr = ActiveMediaPlates[Index];
-		TObjectPtr<AMediaPlate> MediaPlate = PlatePtr.Get();
+		TWeakObjectPtr<UMediaPlateComponent>& PlatePtr = ActiveMediaPlates[Index];
+		TObjectPtr<UMediaPlateComponent> MediaPlate = PlatePtr.Get();
 		bool bIsMediaPlateToBeRemoved = true;
 		if (MediaPlate != nullptr)
 		{
@@ -84,7 +84,7 @@ void FMediaPlateEditorModule::Tick(float DeltaTime)
 	}
 }
 
-void FMediaPlateEditorModule::MediaPlateStartedPlayback(TObjectPtr<AMediaPlate> MediaPlate)
+void FMediaPlateEditorModule::MediaPlateStartedPlayback(TObjectPtr<UMediaPlateComponent> MediaPlate)
 {
 	if (MediaPlate != nullptr)
 	{
diff --git a/Engine/Plugins/Media/MediaPlate/Source/MediaPlateEditor/Private/Toolkits/MediaPlateEditorToolkit.cpp b/Engine/Plugins/Media/MediaPlate/Source/MediaPlateEditor/Private/Toolkits/MediaPlateEditorToolkit.cpp
--- a/Engine/Plugins/Media/MediaPlate/Source/MediaPlateEditor/Private/Toolkits/MediaPlateEditorToolkit.cpp
+++ b/Engine/Plugins/Media/MediaPlate/Source/MediaPlateEditor/Private/Toolkits/MediaPlateEditorToolkit.cpp
@@ -7,6 +7,7 @@
 #include "EditorReimportHandler.h"
 #include "Framework/MultiBox/MultiBoxBuilder.h"
 #include "MediaPlate.h"
+#include "MediaPlateComponent.h"
 #include "MediaPlateEditorModule.h"
 #include "MediaPlayer.h"
 #include "MediaPlaylist.h"
@@ -45,7 +46,7 @@ FMediaPlateEditorToolkit::~FMediaPlateEditorToolkit()
 /* FMediaPlateEditorToolkit interface
  *****************************************************************************/
 
-void FMediaPlateEditorToolkit::Initialize(AMediaPlate* InMediaPlate, const EToolkitMode::Type InMode, const TSharedPtr<IToolkitHost>& InToolkitHost)
+void FMediaPlateEditorToolkit::Initialize(UMediaPlateComponent* InMediaPlate, const EToolkitMode::Type InMode, const TSharedPtr<IToolkitHost>& InToolkitHost)
 {
 	MediaPlate = InMediaPlate;
 
diff --git a/Engine/Plugins/Media/MediaPlate/Source/MediaPlateEditor/Private/Toolkits/MediaPlateEditorToolkit.h b/Engine/Plugins/Media/MediaPlate/Source/MediaPlateEditor/Private/Toolkits/MediaPlateEditorToolkit.h
--- a/Engine/Plugins/Media/MediaPlate/Source/MediaPlateEditor/Private/Toolkits/MediaPlateEditorToolkit.h
+++ b/Engine/Plugins/Media/MediaPlate/Source/MediaPlateEditor/Private/Toolkits/MediaPlateEditorToolkit.h
@@ -9,7 +9,7 @@
 #include "Toolkits/AssetEditorToolkit.h"
 #include "UObject/GCObject.h"
 
-class AMediaPlate;
+class UMediaPlateComponent;
 
 /**
  * Implements an Editor toolkit for media plates.
@@ -33,11 +33,11 @@ public:
 	/**
 	 * Initializes the editor tool kit.
 	 *
-	 * @param InMediaPlate The AMediaPlate asset to edit.
+	 * @param InMediaPlate The media plate asset to edit.
 	 * @param InMode The mode to create the toolkit in.
 	 * @param InToolkitHost The toolkit host.
 	 */
-	void Initialize(AMediaPlate* InMediaPlate, const EToolkitMode::Type InMode, const TSharedPtr<IToolkitHost>& InToolkitHost);
+	void Initialize(UMediaPlateComponent* InMediaPlate, const EToolkitMode::Type InMode, const TSharedPtr<IToolkitHost>& InToolkitHost);
 
 	//~ FAssetEditorToolkit interface
 	virtual FString GetDocumentationLink() const override;
@@ -90,7 +90,7 @@ private:
 	TSharedRef<SDockTab> HandleTabManagerSpawnTab(const FSpawnTabArgs& Args, FName TabIdentifier);
 
 	/** The media plate asset being edited. */
-	AMediaPlate* MediaPlate;
+	TObjectPtr<UMediaPlateComponent> MediaPlate;
 
 	/** Pointer to the style set to use for toolkits. */
 	TSharedRef<ISlateStyle> Style;
diff --git a/Engine/Plugins/Media/MediaPlate/Source/MediaPlateEditor/Private/Widgets/SMediaPlateEditorDetails.cpp b/Engine/Plugins/Media/MediaPlate/Source/MediaPlateEditor/Private/Widgets/SMediaPlateEditorDetails.cpp
--- a/Engine/Plugins/Media/MediaPlate/Source/MediaPlateEditor/Private/Widgets/SMediaPlateEditorDetails.cpp
+++ b/Engine/Plugins/Media/MediaPlate/Source/MediaPlateEditor/Private/Widgets/SMediaPlateEditorDetails.cpp
@@ -2,7 +2,7 @@
 
 #include "Widgets/SMediaPlateEditorDetails.h"
 #include "IDetailsView.h"
-#include "MediaPlate.h"
+#include "MediaPlateComponent.h"
 #include "Modules/ModuleManager.h"
 #include "PropertyEditorModule.h"
 
@@ -11,7 +11,7 @@
 /* SMediaPlateEditorDetails interface
  *****************************************************************************/
 
-void SMediaPlateEditorDetails::Construct(const FArguments& InArgs, AMediaPlate& InMediaPlate, const TSharedRef<ISlateStyle>& InStyle)
+void SMediaPlateEditorDetails::Construct(const FArguments& InArgs, UMediaPlateComponent& InMediaPlate, const TSharedRef<ISlateStyle>& InStyle)
 {
 	MediaPlate = &InMediaPlate;
 
diff --git a/Engine/Plugins/Media/MediaPlate/Source/MediaPlateEditor/Private/Widgets/SMediaPlateEditorDetails.h b/Engine/Plugins/Media/MediaPlate/Source/MediaPlateEditor/Private/Widgets/SMediaPlateEditorDetails.h
--- a/Engine/Plugins/Media/MediaPlate/Source/MediaPlateEditor/Private/Widgets/SMediaPlateEditorDetails.h
+++ b/Engine/Plugins/Media/MediaPlate/Source/MediaPlateEditor/Private/Widgets/SMediaPlateEditorDetails.h
@@ -7,7 +7,7 @@
 #include "Widgets/DeclarativeSyntaxSupport.h"
 #include "Widgets/SCompoundWidget.h"
 
-class AMediaPlate;
+class UMediaPlateComponent;
 
 /**
  * Implements the details panel of the MediaPlate asset editor.
@@ -29,10 +29,10 @@ public:
 	 * @param InMediaPlate	The MediaPlate to show the details for.
 	 * @param InStyleSet	The style set to use.
 	 */
-	void Construct(const FArguments& InArgs, AMediaPlate& InMediaPlate, const TSharedRef<ISlateStyle>& InStyle);
+	void Construct(const FArguments& InArgs, UMediaPlateComponent& InMediaPlate, const TSharedRef<ISlateStyle>& InStyle);
 
 private:
 
 	/** Pointer to the MediaPlate that is being viewed. */
-	TObjectPtr<AMediaPlate> MediaPlate;
+	TObjectPtr<UMediaPlateComponent> MediaPlate;
 };
diff --git a/Engine/Plugins/Media/MediaPlate/Source/MediaPlateEditor/Public/MediaPlateEditorModule.h b/Engine/Plugins/Media/MediaPlate/Source/MediaPlateEditor/Public/MediaPlateEditorModule.h
--- a/Engine/Plugins/Media/MediaPlate/Source/MediaPlateEditor/Public/MediaPlateEditorModule.h
+++ b/Engine/Plugins/Media/MediaPlate/Source/MediaPlateEditor/Public/MediaPlateEditorModule.h
@@ -8,10 +8,10 @@
 #include "TickableEditorObject.h"
 #include "UObject/ObjectPtr.h"
 
-class AMediaPlate;
 class IAssetTools;
 class IAssetTypeActions;
 class ISlateStyle;
+class UMediaPlateComponent;
 
 /** Log category for this module. */
 DECLARE_LOG_CATEGORY_EXTERN(LogMediaPlateEditor, Log, All);
@@ -31,14 +31,14 @@ public:
 	/**
 	 * Call this when a media plate starts playing so we can track it.
 	 */
-	void MediaPlateStartedPlayback(TObjectPtr<AMediaPlate> MediaPlate);
+	void MediaPlateStartedPlayback(TObjectPtr<UMediaPlateComponent> MediaPlate);
 	
 private:
 
 	/** Customization name to avoid reusing staticstruct during shutdown. */
 	FName MediaPlateName;
 	/** Holds all the media plates that are playing. */
-	TArray<TWeakObjectPtr<AMediaPlate>> ActiveMediaPlates;
+	TArray<TWeakObjectPtr<UMediaPlateComponent>> ActiveMediaPlates;
 
 	/** Holds the plug-ins style set. */
 	TSharedPtr<ISlateStyle> Style;
