MediaPlate: Uses smart caching for image sequences.

#rnx

diff --git a/Engine/Plugins/Media/MediaPlate/Source/MediaPlate/Private/MediaPlateComponent.cpp b/Engine/Plugins/Media/MediaPlate/Source/MediaPlate/Private/MediaPlateComponent.cpp
--- a/Engine/Plugins/Media/MediaPlate/Source/MediaPlate/Private/MediaPlateComponent.cpp
+++ b/Engine/Plugins/Media/MediaPlate/Source/MediaPlate/Private/MediaPlateComponent.cpp
@@ -79,14 +79,15 @@ void UMediaPlateComponent::Play()
 			MediaPathMediaSource = UMediaSource::SpawnMediaSourceForString(MediaPath.FilePath);
 			if (MediaPathMediaSource != nullptr)
 			{
+				SetMediaSourceOptions(MediaPathMediaSource);
 				bIsPlaying = MediaPlayer->OpenSource(MediaPathMediaSource);
-				return;
 			}
 		}
 
 		// If we did not get anything, try the media source.
 		if (bIsPlaying == false)
 		{
+			SetMediaSourceOptions(MediaSource);
 			bIsPlaying = MediaPlayer->OpenSource(MediaSource);
 		}
 	}
@@ -134,4 +135,13 @@ void UMediaPlateComponent::UnregisterWithMediaTextureTracker()
 	}
 }
 
+void UMediaPlateComponent::SetMediaSourceOptions(UMediaSource* InMediaSource)
+{
+	if (InMediaSource != nullptr)
+	{
+		InMediaSource->SetMediaOptionBool(TEXT("ImgMediaSmartCacheEnabled"), bSmartCacheEnabled);
+		InMediaSource->SetMediaOptionFloat(TEXT("ImgMediaSmartCacheTimeToLookAhead"), SmartCacheTimeToLookAhead);
+	}
+}
+
 #undef LOCTEXT_NAMESPACE
diff --git a/Engine/Plugins/Media/MediaPlate/Source/MediaPlate/Public/MediaPlateComponent.h b/Engine/Plugins/Media/MediaPlate/Source/MediaPlate/Public/MediaPlateComponent.h
--- a/Engine/Plugins/Media/MediaPlate/Source/MediaPlate/Public/MediaPlateComponent.h
+++ b/Engine/Plugins/Media/MediaPlate/Source/MediaPlate/Public/MediaPlateComponent.h
@@ -68,6 +68,18 @@ public:
 	UPROPERTY(EditAnywhere, Category = MediaPlate)
 	TObjectPtr<UMediaSource> MediaSource;
 
+	/** Enable smart caching for image sequences. */
+	UPROPERTY(EditAnywhere, AdvancedDisplay, Category = "MediaPlate")
+	bool bSmartCacheEnabled = true;
+
+	/**
+	 * The cache will fill up with frames that are up to this time from the current time.
+	 * E.g. if this is 0.2, and we are at time index 5 seconds,
+	 * then we will fill the cache with frames between 5 seconds and 5.2 seconds.
+	 */
+	UPROPERTY(EditAnywhere, AdvancedDisplay, Category = "MediaPlate")
+	float SmartCacheTimeToLookAhead = 0.2f;
+
 	/**
 	 * Adds our media texture to the media texture tracker.
 	 */
@@ -87,4 +99,9 @@ private:
 	/** If we are using MediaPath, then this is the media source for it. */
 	UPROPERTY(Transient)
 	TObjectPtr<UMediaSource> MediaPathMediaSource;
+
+	/**
+	 * Apply any options we need to the media source.
+	 */
+	void SetMediaSourceOptions(UMediaSource* InMediaSource);
 };
