MediaPlate: Consolidated code that plays a media source.
#rnx

diff --git a/Engine/Plugins/Media/MediaPlate/Source/MediaPlate/Private/MediaPlateComponent.cpp b/Engine/Plugins/Media/MediaPlate/Source/MediaPlate/Private/MediaPlateComponent.cpp
--- a/Engine/Plugins/Media/MediaPlate/Source/MediaPlate/Private/MediaPlateComponent.cpp
+++ b/Engine/Plugins/Media/MediaPlate/Source/MediaPlate/Private/MediaPlateComponent.cpp
@@ -79,16 +79,14 @@ void UMediaPlateComponent::Play()
 			MediaPathMediaSource = UMediaSource::SpawnMediaSourceForString(MediaPath.FilePath);
 			if (MediaPathMediaSource != nullptr)
 			{
-				SetMediaSourceOptions(MediaPathMediaSource);
-				bIsPlaying = MediaPlayer->OpenSource(MediaPathMediaSource);
+				bIsPlaying = PlayMediaSource(MediaPathMediaSource);
 			}
 		}
 
 		// If we did not get anything, try the media source.
 		if (bIsPlaying == false)
 		{
-			SetMediaSourceOptions(MediaSource);
-			bIsPlaying = MediaPlayer->OpenSource(MediaSource);
+			bIsPlaying = PlayMediaSource(MediaSource);
 		}
 	}
 
@@ -135,13 +133,25 @@ void UMediaPlateComponent::UnregisterWithMediaTextureTracker()
 	}
 }
 
-void UMediaPlateComponent::SetMediaSourceOptions(UMediaSource* InMediaSource)
+bool UMediaPlateComponent::PlayMediaSource(UMediaSource* InMediaSource)
 {
+	bool bIsPlaying = false;
+
 	if (InMediaSource != nullptr)
 	{
+		// Set media options.
 		InMediaSource->SetMediaOptionBool(TEXT("ImgMediaSmartCacheEnabled"), bSmartCacheEnabled);
 		InMediaSource->SetMediaOptionFloat(TEXT("ImgMediaSmartCacheTimeToLookAhead"), SmartCacheTimeToLookAhead);
+	
+		TObjectPtr<UMediaPlayer> MediaPlayer = GetMediaPlayer();
+		if (MediaPlayer != nullptr)
+		{
+			// Play the source.
+			bIsPlaying = MediaPlayer->OpenSource(InMediaSource);
+		}
 	}
+
+	return bIsPlaying;
 }
 
 #undef LOCTEXT_NAMESPACE
diff --git a/Engine/Plugins/Media/MediaPlate/Source/MediaPlate/Public/MediaPlateComponent.h b/Engine/Plugins/Media/MediaPlate/Source/MediaPlate/Public/MediaPlateComponent.h
--- a/Engine/Plugins/Media/MediaPlate/Source/MediaPlate/Public/MediaPlateComponent.h
+++ b/Engine/Plugins/Media/MediaPlate/Source/MediaPlate/Public/MediaPlateComponent.h
@@ -101,7 +101,10 @@ private:
 	TObjectPtr<UMediaSource> MediaPathMediaSource;
 
 	/**
-	 * Apply any options we need to the media source.
+	 * Plays a media source.
+	 * 
+	 * @param	InMediaSource		Media source to play.
+	 * @return	True if we played anything.
 	 */
-	void SetMediaSourceOptions(UMediaSource* InMediaSource);
+	bool PlayMediaSource(UMediaSource* InMediaSource);
 };
