MediaPlate: Button to open media plate window only shows if you are not already in that window.
#rnx

diff --git a/Engine/Plugins/Media/MediaPlate/Source/MediaPlateEditor/Private/MediaPlateCustomization.cpp b/Engine/Plugins/Media/MediaPlate/Source/MediaPlateEditor/Private/MediaPlateCustomization.cpp
--- a/Engine/Plugins/Media/MediaPlate/Source/MediaPlateEditor/Private/MediaPlateCustomization.cpp
+++ b/Engine/Plugins/Media/MediaPlate/Source/MediaPlateEditor/Private/MediaPlateCustomization.cpp
@@ -20,6 +20,11 @@
 
 void FMediaPlateCustomization::CustomizeDetails(IDetailLayoutBuilder& DetailBuilder)
 {
+	// Is this the media plate editor window?
+	const IDetailsView* DetailsView = DetailBuilder.GetDetailsView();
+	TSharedPtr<FTabManager> HostTabManager = DetailsView->GetHostTabManager();
+	bool bIsMediaPlateWindow = (HostTabManager.IsValid() == false);
+
 	IDetailCategoryBuilder& MediaPlateCategory = DetailBuilder.EditCategory("MediaPlate");
 
 	// Get objects we are editing.
@@ -141,23 +146,26 @@ void FMediaPlateCustomization::CustomizeDetails(IDetailLayoutBuilder& DetailBuil
 
 
 	// Add button to open the media plate editor.
-	MediaPlateCategory.AddCustomRow(LOCTEXT("OpenMediaPlate", "Open Media Plate"))
-		[
-			SNew(SHorizontalBox)
-
-			+ SHorizontalBox::Slot()
-				.FillWidth(1.f)
-				.Padding(0, 5, 10, 5)
-				[
-					SNew(SButton)
-						.ContentPadding(3)
-						.VAlign(VAlign_Center)
-						.HAlign(HAlign_Center)
-						.OnClicked(this, &FMediaPlateCustomization::OnOpenMediaPlate)
-						.Text(LOCTEXT("OpenMediaPlate", "Open Media Plate"))
-				]
-		];
-}
+	if (bIsMediaPlateWindow == false)
+	{
+		MediaPlateCategory.AddCustomRow(LOCTEXT("OpenMediaPlate", "Open Media Plate"))
+			[
+				SNew(SHorizontalBox)
+
+				+ SHorizontalBox::Slot()
+					.FillWidth(1.f)
+					.Padding(0, 5, 10, 5)
+					[
+						SNew(SButton)
+							.ContentPadding(3)
+							.VAlign(VAlign_Center)
+							.HAlign(HAlign_Center)
+							.OnClicked(this, &FMediaPlateCustomization::OnOpenMediaPlate)
+							.Text(LOCTEXT("OpenMediaPlate", "Open Media Plate"))
+					]
+			];
+		}
+	}
 
 FReply FMediaPlateCustomization::OnOpenMediaPlate()
 {
