MediaPlate: Added start time.
#rnx

diff --git a/Engine/Plugins/Media/MediaPlate/Source/MediaPlate/MediaPlate.Build.cs b/Engine/Plugins/Media/MediaPlate/Source/MediaPlate/MediaPlate.Build.cs
--- a/Engine/Plugins/Media/MediaPlate/Source/MediaPlate/MediaPlate.Build.cs
+++ b/Engine/Plugins/Media/MediaPlate/Source/MediaPlate/MediaPlate.Build.cs
@@ -21,6 +21,7 @@ public class MediaPlate : ModuleRules
 				"CoreUObject",
 				"Engine",
 				"MediaAssets",
+				"MediaUtils",
 			}
 			);
 	}
diff --git a/Engine/Plugins/Media/MediaPlate/Source/MediaPlate/Private/MediaPlateComponent.cpp b/Engine/Plugins/Media/MediaPlate/Source/MediaPlate/Private/MediaPlateComponent.cpp
--- a/Engine/Plugins/Media/MediaPlate/Source/MediaPlate/Private/MediaPlateComponent.cpp
+++ b/Engine/Plugins/Media/MediaPlate/Source/MediaPlate/Private/MediaPlateComponent.cpp
@@ -70,9 +70,6 @@ void UMediaPlateComponent::Play()
 	TObjectPtr<UMediaPlayer> MediaPlayer = GetMediaPlayer();
 	if (MediaPlayer != nullptr)
 	{
-		MediaPlayer->PlayOnOpen = true;
-		MediaPlayer->SetLooping(bLoop);
-		
 		// Try and play the media path.
 		if (MediaPath.FilePath.IsEmpty() == false)
 		{
@@ -147,7 +144,12 @@ bool UMediaPlateComponent::PlayMediaSource(UMediaSource* InMediaSource)
 		if (MediaPlayer != nullptr)
 		{
 			// Play the source.
-			bIsPlaying = MediaPlayer->OpenSource(InMediaSource);
+			FMediaPlayerOptions Options;
+			Options.SeekTime = FTimespan::FromSeconds(StartTime);
+			Options.PlayOnOpen = EMediaPlayerOptionBooleanOverride::Enabled;
+			Options.Loop = bLoop ? EMediaPlayerOptionBooleanOverride::Enabled :
+				EMediaPlayerOptionBooleanOverride::Disabled;
+			bIsPlaying = MediaPlayer->OpenSourceWithOptions(InMediaSource, Options);
 		}
 	}
 
diff --git a/Engine/Plugins/Media/MediaPlate/Source/MediaPlate/Public/MediaPlateComponent.h b/Engine/Plugins/Media/MediaPlate/Source/MediaPlate/Public/MediaPlateComponent.h
--- a/Engine/Plugins/Media/MediaPlate/Source/MediaPlate/Public/MediaPlateComponent.h
+++ b/Engine/Plugins/Media/MediaPlate/Source/MediaPlate/Public/MediaPlateComponent.h
@@ -4,6 +4,8 @@
 
 #include "Components/ActorComponent.h"
 #include "CoreMinimal.h"
+#include "MediaPlayerOptions.h"
+
 #include "MediaPlateComponent.generated.h"
 
 class UMediaComponent;
@@ -56,6 +58,10 @@ public:
 	UPROPERTY(EditAnywhere, Category = "MediaPlate")
 	bool bLoop;
 
+	/** What time to start playing from (in seconds). */
+	UPROPERTY(EditAnywhere, Category = "MediaPlate", meta = (ClampMin = "0.0"))
+	float StartTime = 0.0f;
+
 	/** Holds the media player. */
 	UPROPERTY(Category = MediaPlate, VisibleAnywhere, BlueprintReadOnly)
 	TObjectPtr<UMediaComponent> MediaComponent;
