MediaPlate: Added function to get media plate class.
#rnx

diff --git a/Engine/Plugins/Media/MediaPlate/Source/MediaPlate/Private/MediaPlateModule.cpp b/Engine/Plugins/Media/MediaPlate/Source/MediaPlate/Private/MediaPlateModule.cpp
--- a/Engine/Plugins/Media/MediaPlate/Source/MediaPlate/Private/MediaPlateModule.cpp
+++ b/Engine/Plugins/Media/MediaPlate/Source/MediaPlate/Private/MediaPlateModule.cpp
@@ -2,8 +2,15 @@
 
 #include "MediaPlateModule.h"
 
+#include "MediaPlate.h"
+
 DEFINE_LOG_CATEGORY(LogMediaPlate);
 
+UClass* FMediaPlateModule::GetAMediaPlateClass()
+{
+	return AMediaPlate::StaticClass();
+}
+
 void FMediaPlateModule::StartupModule()
 {
 }
diff --git a/Engine/Plugins/Media/MediaPlate/Source/MediaPlate/Public/MediaPlateModule.h b/Engine/Plugins/Media/MediaPlate/Source/MediaPlate/Public/MediaPlateModule.h
--- a/Engine/Plugins/Media/MediaPlate/Source/MediaPlate/Public/MediaPlateModule.h
+++ b/Engine/Plugins/Media/MediaPlate/Source/MediaPlate/Public/MediaPlateModule.h
@@ -9,10 +9,14 @@
 /** Log category for this module. */
 DECLARE_LOG_CATEGORY_EXTERN(LogMediaPlate, Log, All);
 
-class FMediaPlateModule : public IModuleInterface
+class MEDIAPLATE_API FMediaPlateModule : public IModuleInterface
 {
 public:
-	
+	/**
+	 * Call this to get the UClass for AMediaPlate.
+	 */
+	virtual UClass* GetAMediaPlateClass();
+
 	/** IModuleInterface implementation */
 	virtual void StartupModule() override;
 	virtual void ShutdownModule() override;
