MediaPlate: Fix to allow material overrides, where cast to UMaterialInstanceDynamic would return nullptr previously.

diff --git a/Engine/Plugins/Media/MediaPlate/Source/MediaPlate/Private/MediaPlate.cpp b/Engine/Plugins/Media/MediaPlate/Source/MediaPlate/Private/MediaPlate.cpp
--- a/Engine/Plugins/Media/MediaPlate/Source/MediaPlate/Private/MediaPlate.cpp
+++ b/Engine/Plugins/Media/MediaPlate/Source/MediaPlate/Private/MediaPlate.cpp
@@ -52,25 +52,27 @@ void AMediaPlate::PostRegisterAllComponents()
 	{
 		if (StaticMeshComponent != nullptr)
 		{
+			UMaterialInstanceDynamic* MediaMID = nullptr;
+
 			// Do we have a material?
-			if (StaticMeshComponent->GetNumOverrideMaterials() == 0)
+			if (StaticMeshComponent->GetNumOverrideMaterials() > 0)
+			{
+				MediaMID = StaticMeshComponent->CreateAndSetMaterialInstanceDynamicFromMaterial(0, StaticMeshComponent->GetMaterial(0));
+			}
+			else
 			{
-				// Add material.
+				// Add default material.
 				UMaterial* Material = LoadObject<UMaterial>(nullptr, TEXT("/MediaPlate/M_MediaPlate"));
 				if (Material != nullptr)
 				{
-					UMaterialInstanceDynamic* MaterialInstance = StaticMeshComponent->CreateAndSetMaterialInstanceDynamicFromMaterial(0, Material);
+					MediaMID = StaticMeshComponent->CreateAndSetMaterialInstanceDynamicFromMaterial(0, Material);
 				}
 			}
 
 			// Set up the material to point to our media texture.
-			if (StaticMeshComponent->GetNumMaterials() > 0)
+			if (MediaMID != nullptr)
 			{
-				UMaterialInstanceDynamic* MaterialInstance = Cast< UMaterialInstanceDynamic>(StaticMeshComponent->GetMaterial(0));
-				if (MaterialInstance != nullptr)
-				{
-					MaterialInstance->SetTextureParameterValue(MediaTextureName, MediaPlateComponent->GetMediaTexture());
-				}
+				MediaMID->SetTextureParameterValue(MediaTextureName, MediaPlateComponent->GetMediaTexture());
 			}
 		}
 
