MediaPlate: Added helper function to get a media player from an object.
#rnx

diff --git a/Engine/Plugins/Media/MediaPlate/Source/MediaPlate/Private/MediaPlateModule.cpp b/Engine/Plugins/Media/MediaPlate/Source/MediaPlate/Private/MediaPlateModule.cpp
--- a/Engine/Plugins/Media/MediaPlate/Source/MediaPlate/Private/MediaPlateModule.cpp
+++ b/Engine/Plugins/Media/MediaPlate/Source/MediaPlate/Private/MediaPlateModule.cpp
@@ -3,6 +3,7 @@
 #include "MediaPlateModule.h"
 
 #include "MediaPlate.h"
+#include "MediaPlateComponent.h"
 
 DEFINE_LOG_CATEGORY(LogMediaPlate);
 
@@ -11,6 +12,23 @@ UClass* FMediaPlateModule::GetAMediaPlateClass()
 	return AMediaPlate::StaticClass();
 }
 
+UMediaPlayer* FMediaPlateModule::GetMediaPlayer(UObject* Object)
+{
+	UMediaPlayer* MediaPlayer = nullptr;
+
+	AMediaPlate* MediaPlate = Cast<AMediaPlate>(Object);
+	if (MediaPlate != nullptr)
+	{
+		UMediaPlateComponent* MediaPlateComponent = MediaPlate->MediaPlateComponent;
+		if (MediaPlateComponent != nullptr)
+		{
+			MediaPlayer = MediaPlateComponent->GetMediaPlayer();
+		}
+	}
+
+	return MediaPlayer;
+}
+
 void FMediaPlateModule::StartupModule()
 {
 }
diff --git a/Engine/Plugins/Media/MediaPlate/Source/MediaPlate/Public/MediaPlateModule.h b/Engine/Plugins/Media/MediaPlate/Source/MediaPlate/Public/MediaPlateModule.h
--- a/Engine/Plugins/Media/MediaPlate/Source/MediaPlate/Public/MediaPlateModule.h
+++ b/Engine/Plugins/Media/MediaPlate/Source/MediaPlate/Public/MediaPlateModule.h
@@ -6,6 +6,8 @@
 #include "Logging/LogMacros.h"
 #include "Modules/ModuleManager.h"
 
+class UMediaPlayer;
+
 /** Log category for this module. */
 DECLARE_LOG_CATEGORY_EXTERN(LogMediaPlate, Log, All);
 
@@ -17,6 +19,11 @@ public:
 	 */
 	virtual UClass* GetAMediaPlateClass();
 
+	/**
+	 * Call this to get the media player from a media plate object.
+	 */
+	virtual UMediaPlayer* GetMediaPlayer(UObject* Object);
+
 	/** IModuleInterface implementation */
 	virtual void StartupModule() override;
 	virtual void ShutdownModule() override;
