MedaPlate: Added sound.

#rnx

diff --git a/Engine/Plugins/Media/MediaPlate/Source/MediaPlate/MediaPlate.Build.cs b/Engine/Plugins/Media/MediaPlate/Source/MediaPlate/MediaPlate.Build.cs
--- a/Engine/Plugins/Media/MediaPlate/Source/MediaPlate/MediaPlate.Build.cs
+++ b/Engine/Plugins/Media/MediaPlate/Source/MediaPlate/MediaPlate.Build.cs
@@ -18,6 +18,7 @@ public class MediaPlate : ModuleRules
 		PrivateDependencyModuleNames.AddRange(
 			new string[]
 			{
+				"AudioMixer",
 				"CoreUObject",
 				"Engine",
 				"MediaAssets",
diff --git a/Engine/Plugins/Media/MediaPlate/Source/MediaPlate/Private/MediaPlateComponent.cpp b/Engine/Plugins/Media/MediaPlate/Source/MediaPlate/Private/MediaPlateComponent.cpp
--- a/Engine/Plugins/Media/MediaPlate/Source/MediaPlate/Private/MediaPlateComponent.cpp
+++ b/Engine/Plugins/Media/MediaPlate/Source/MediaPlate/Private/MediaPlateComponent.cpp
@@ -8,6 +8,7 @@
 #include "MediaPlayer.h"
 #include "MediaPlateModule.h"
 #include "MediaSource.h"
+#include "MediaSoundComponent.h"
 #include "MediaTexture.h"
 #include "MediaTextureTracker.h"
 
@@ -31,6 +32,21 @@ UMediaPlateComponent::UMediaPlateComponent(const FObjectInitializer& ObjectIniti
 	}
 }
 
+void UMediaPlateComponent::OnRegister()
+{
+	Super::OnRegister();
+
+	// Set up sound component if we have one.
+	if (SoundComponent != nullptr)
+	{
+		TObjectPtr<UMediaPlayer> MediaPlayer = GetMediaPlayer();
+		if (MediaPlayer != nullptr)
+		{
+			SoundComponent->SetMediaPlayer(MediaPlayer);
+		}
+	}
+}
+
 void UMediaPlateComponent::BeginPlay()
 {
 	Super::BeginPlay();
@@ -156,4 +172,46 @@ bool UMediaPlateComponent::PlayMediaSource(UMediaSource* InMediaSource)
 	return bIsPlaying;
 }
 
+#if WITH_EDITOR
+
+void UMediaPlateComponent::PostEditChangeProperty(FPropertyChangedEvent& PropertyChangedEvent)
+{
+	// Has bEnableAudiio changed?
+	if (PropertyChangedEvent.GetPropertyName() == GET_MEMBER_NAME_CHECKED(ThisClass, bEnableAudio))
+	{
+		// Are we turning on audio?
+		if (bEnableAudio)
+		{
+			// Get the media player.
+			TObjectPtr<UMediaPlayer> MediaPlayer = GetMediaPlayer();
+			if (MediaPlayer != nullptr)
+			{
+				// Create a sound component.
+				SoundComponent = NewObject<UMediaSoundComponent>(this, NAME_None);
+				if (SoundComponent != nullptr)
+				{
+					SoundComponent->bIsUISound = true;
+					SoundComponent->SetMediaPlayer(MediaPlayer);
+					SoundComponent->Initialize();
+					SoundComponent->RegisterComponent();
+				}
+			}
+		}
+		else
+		{
+			// Remove this sound component.
+			if (SoundComponent != nullptr)
+			{
+				SoundComponent->UnregisterComponent();
+				SoundComponent->SetMediaPlayer(nullptr);
+				SoundComponent->UpdatePlayer();
+				SoundComponent->DestroyComponent();
+				SoundComponent = nullptr;
+			}
+		}
+	}
+}
+
+#endif // WITH_EDITOR
+
 #undef LOCTEXT_NAMESPACE
diff --git a/Engine/Plugins/Media/MediaPlate/Source/MediaPlate/Public/MediaPlateComponent.h b/Engine/Plugins/Media/MediaPlate/Source/MediaPlate/Public/MediaPlateComponent.h
--- a/Engine/Plugins/Media/MediaPlate/Source/MediaPlate/Public/MediaPlateComponent.h
+++ b/Engine/Plugins/Media/MediaPlate/Source/MediaPlate/Public/MediaPlateComponent.h
@@ -10,6 +10,7 @@
 
 class UMediaComponent;
 class UMediaPlayer;
+class UMediaSoundComponent;
 class UMediaSource;
 class UMediaTexture;
 struct FMediaTextureTrackerObject;
@@ -24,7 +25,11 @@ class MEDIAPLATE_API UMediaPlateComponent : public UActorComponent
 
 public:
 	//~ UActorComponent interface.
+	virtual void OnRegister();
 	virtual void BeginPlay();
+#if WITH_EDITOR
+	virtual void PostEditChangeProperty(FPropertyChangedEvent& PropertyChangedEvent) override;
+#endif // WITH_EDITOR
 
 	/**
 	 * Call this get our media player.
@@ -58,6 +63,10 @@ public:
 	UPROPERTY(EditAnywhere, Category = "MediaPlate")
 	bool bLoop;
 
+	/** If set then enable audio. */
+	UPROPERTY(EditAnywhere, Category = "MediaPlate")
+	bool bEnableAudio = false;
+
 	/** What time to start playing from (in seconds). */
 	UPROPERTY(EditAnywhere, Category = "MediaPlate", meta = (ClampMin = "0.0"))
 	float StartTime = 0.0f;
@@ -66,6 +75,10 @@ public:
 	UPROPERTY(Category = MediaPlate, VisibleAnywhere, BlueprintReadOnly)
 	TObjectPtr<UMediaComponent> MediaComponent;
 
+	/** Holds the component to play sound. */
+	UPROPERTY(EditAnywhere, AdvancedDisplay, Category = MediaPlate)
+	TObjectPtr<UMediaSoundComponent> SoundComponent;
+
 	/** URL (or file)  to play. This will take precedence over the MediaSource. */
 	UPROPERTY(EditAnywhere, Category = MediaPlate)
 	FFilePath MediaPath;
diff --git a/Engine/Plugins/Media/MediaPlate/Source/MediaPlateEditor/Private/MediaPlateEditorModule.cpp b/Engine/Plugins/Media/MediaPlate/Source/MediaPlateEditor/Private/MediaPlateEditorModule.cpp
--- a/Engine/Plugins/Media/MediaPlate/Source/MediaPlateEditor/Private/MediaPlateEditorModule.cpp
+++ b/Engine/Plugins/Media/MediaPlate/Source/MediaPlateEditor/Private/MediaPlateEditorModule.cpp
@@ -9,6 +9,7 @@
 #include "MediaPlateCustomization.h"
 #include "MediaPlateEditorStyle.h"
 #include "MediaPlayer.h"
+#include "MediaSoundComponent.h"
 #include "Models/MediaPlateEditorCommands.h"
 #include "PropertyEditorModule.h"
 
@@ -52,6 +53,14 @@ void FMediaPlateEditorModule::Tick(float DeltaTime)
 		bool bIsMediaPlateToBeRemoved = true;
 		if (MediaPlate != nullptr)
 		{
+			// Update sound component.
+			TObjectPtr<UMediaSoundComponent> SoundComponent = MediaPlate->SoundComponent;
+			if (SoundComponent != nullptr)
+			{
+				SoundComponent->UpdatePlayer();
+			}
+
+			// Get the player.
 			TObjectPtr<UMediaPlayer> MediaPlayer = MediaPlate->GetMediaPlayer();
 			if (MediaPlayer != nullptr)
 			{
