MediaPlate: Removed empty space in front of play buttons.
#rnx

diff --git a/Engine/Plugins/Media/MediaPlate/Source/MediaPlateEditor/Private/MediaPlateCustomization.cpp b/Engine/Plugins/Media/MediaPlate/Source/MediaPlateEditor/Private/MediaPlateCustomization.cpp
--- a/Engine/Plugins/Media/MediaPlate/Source/MediaPlateEditor/Private/MediaPlateCustomization.cpp
+++ b/Engine/Plugins/Media/MediaPlate/Source/MediaPlateEditor/Private/MediaPlateCustomization.cpp
@@ -73,13 +73,6 @@ void FMediaPlateCustomization::CustomizeDetails(IDetailLayoutBuilder& DetailBuil
 
 	// Add media control buttons.
 	MediaPlateCategory.AddCustomRow(LOCTEXT("MediaPlateControls", "MediaPlate Controls"))
-		.NameContent()
-		[
-			SNullWidget::NullWidget
-		]
-		.ValueContent()
-		.VAlign(VAlign_Center)
-		.MaxDesiredWidth(250)
 		[
 			SNew(SHorizontalBox)
 
