MediaPlate: Customization uses button images from the media plate style.
#rnx

diff --git a/Engine/Plugins/Media/MediaPlate/Source/MediaPlateEditor/Private/MediaPlateCustomization.cpp b/Engine/Plugins/Media/MediaPlate/Source/MediaPlateEditor/Private/MediaPlateCustomization.cpp
--- a/Engine/Plugins/Media/MediaPlate/Source/MediaPlateEditor/Private/MediaPlateCustomization.cpp
+++ b/Engine/Plugins/Media/MediaPlate/Source/MediaPlateEditor/Private/MediaPlateCustomization.cpp
@@ -25,6 +25,18 @@ void FMediaPlateCustomization::CustomizeDetails(IDetailLayoutBuilder& DetailBuil
 	TSharedPtr<FTabManager> HostTabManager = DetailsView->GetHostTabManager();
 	bool bIsMediaPlateWindow = (HostTabManager.IsValid() == false);
 
+	// Get style.
+	const ISlateStyle* Style = nullptr;
+	FMediaPlateEditorModule* EditorModule = FModuleManager::LoadModulePtr<FMediaPlateEditorModule>("MediaPlateEditor");
+	if (EditorModule != nullptr)
+	{
+		Style = EditorModule->GetStyle().Get();
+	}
+	if (Style == nullptr)
+	{
+		Style = &FEditorStyle::Get();
+	}
+
 	IDetailCategoryBuilder& MediaPlateCategory = DetailBuilder.EditCategory("MediaPlate");
 
 	// Get objects we are editing.
@@ -106,7 +118,7 @@ void FMediaPlateCustomization::CustomizeDetails(IDetailLayoutBuilder& DetailBuil
 						[
 							SNew(SImage)
 								.ColorAndOpacity(FSlateColor::UseForeground())
-								.Image(FEditorStyle::GetBrush("Icons.Toolbar.Play"))
+								.Image(Style->GetBrush("MediaPlateEditor.PlayMedia.Small"))
 						]
 				]
 
@@ -132,7 +144,7 @@ void FMediaPlateCustomization::CustomizeDetails(IDetailLayoutBuilder& DetailBuil
 						[
 							SNew(SImage)
 								.ColorAndOpacity(FSlateColor::UseForeground())
-								.Image(FEditorStyle::GetBrush("Icons.Toolbar.Stop"))
+								.Image(Style->GetBrush("MediaPlateEditor.StopMedia.Small"))
 						]
 				]
 		];
diff --git a/Engine/Plugins/Media/MediaPlate/Source/MediaPlateEditor/Public/MediaPlateEditorModule.h b/Engine/Plugins/Media/MediaPlate/Source/MediaPlateEditor/Public/MediaPlateEditorModule.h
--- a/Engine/Plugins/Media/MediaPlate/Source/MediaPlateEditor/Public/MediaPlateEditorModule.h
+++ b/Engine/Plugins/Media/MediaPlate/Source/MediaPlateEditor/Public/MediaPlateEditorModule.h
@@ -28,6 +28,11 @@ public:
 	virtual void Tick(float DeltaTime) override;
 	virtual TStatId GetStatId() const override { RETURN_QUICK_DECLARE_CYCLE_STAT(FMediaPlateEditorModule, STATGROUP_Tickables); }
 	
+	/**
+	 * Get the style used by this module.
+	 **/
+	TSharedPtr<ISlateStyle> GetStyle() { return Style; }
+	
 	/**
 	 * Call this when a media plate starts playing so we can track it.
 	 */
