MediaPlate: Added pause and forward buttons to customization.
#rnx

diff --git a/Engine/Plugins/Media/MediaPlate/Source/MediaPlateEditor/Private/MediaPlateCustomization.cpp b/Engine/Plugins/Media/MediaPlate/Source/MediaPlateEditor/Private/MediaPlateCustomization.cpp
--- a/Engine/Plugins/Media/MediaPlate/Source/MediaPlateEditor/Private/MediaPlateCustomization.cpp
+++ b/Engine/Plugins/Media/MediaPlate/Source/MediaPlateEditor/Private/MediaPlateCustomization.cpp
@@ -9,6 +9,7 @@
 #include "MediaPlate.h"
 #include "MediaPlateComponent.h"
 #include "MediaPlateEditorModule.h"
+#include "MediaPlayer.h"
 #include "Subsystems/AssetEditorSubsystem.h"
 #include "Widgets/Input/SButton.h"
 #include "Widgets/Images/SImage.h"
@@ -122,6 +123,66 @@ void FMediaPlateCustomization::CustomizeDetails(IDetailLayoutBuilder& DetailBuil
 						]
 				]
 
+			// Pause button.
+			+ SHorizontalBox::Slot()
+				.VAlign(VAlign_Center)
+				.AutoWidth()
+				[
+					SNew(SButton)
+					.VAlign(VAlign_Center)
+					.OnClicked_Lambda([this]() -> FReply
+					{
+						for (TWeakObjectPtr<UMediaPlateComponent>& MediaPlatePtr : MediaPlatesList)
+						{
+							UMediaPlateComponent* MediaPlate = MediaPlatePtr.Get();
+							if (MediaPlate != nullptr)
+							{
+								UMediaPlayer* MediaPlayer = MediaPlate->GetMediaPlayer();
+								if (MediaPlayer != nullptr)
+								{
+									MediaPlayer->Pause();
+								}
+							}
+						}
+						return FReply::Handled();
+					})
+					[
+						SNew(SImage)
+							.ColorAndOpacity(FSlateColor::UseForeground())
+							.Image(Style->GetBrush("MediaPlateEditor.PauseMedia.Small"))
+					]
+				]
+
+			// Forward button.
+			+ SHorizontalBox::Slot()
+				.VAlign(VAlign_Center)
+				.AutoWidth()
+				[
+					SNew(SButton)
+					.VAlign(VAlign_Center)
+					.OnClicked_Lambda([this]() -> FReply
+					{
+						for (TWeakObjectPtr<UMediaPlateComponent>& MediaPlatePtr : MediaPlatesList)
+						{
+							UMediaPlateComponent* MediaPlate = MediaPlatePtr.Get();
+							if (MediaPlate != nullptr)
+							{
+								UMediaPlayer* MediaPlayer = MediaPlate->GetMediaPlayer();
+								if (MediaPlayer != nullptr)
+								{
+									MediaPlayer->SetRate(GetForwardRate(MediaPlayer));
+								}
+							}
+						}
+						return FReply::Handled();
+					})
+					[
+						SNew(SImage)
+							.ColorAndOpacity(FSlateColor::UseForeground())
+							.Image(Style->GetBrush("MediaPlateEditor.ForwardMedia.Small"))
+					]
+				]
+
 			// Stop button.
 			+ SHorizontalBox::Slot()
 				.VAlign(VAlign_Center)
@@ -200,4 +261,18 @@ void FMediaPlateCustomization::OnMediaPathChanged(IDetailLayoutBuilder* DetailBu
 	DetailBuilder->ForceRefreshDetails();
 }
 
+
+float FMediaPlateCustomization::GetForwardRate(UMediaPlayer* MediaPlayer) const
+{
+	float Rate = MediaPlayer->GetRate();
+
+	if (Rate < 1.0f)
+	{
+		Rate = 1.0f;
+	}
+
+	return 2.0f * Rate;
+}
+
+
 #undef LOCTEXT_NAMESPACE
diff --git a/Engine/Plugins/Media/MediaPlate/Source/MediaPlateEditor/Private/MediaPlateCustomization.h b/Engine/Plugins/Media/MediaPlate/Source/MediaPlateEditor/Private/MediaPlateCustomization.h
--- a/Engine/Plugins/Media/MediaPlate/Source/MediaPlateEditor/Private/MediaPlateCustomization.h
+++ b/Engine/Plugins/Media/MediaPlate/Source/MediaPlateEditor/Private/MediaPlateCustomization.h
@@ -9,6 +9,7 @@
 
 class IDetailLayoutBuilder;
 class UMediaPlateComponent;
+class UMediaPlayer;
 
 /**
  * Implements a details view customization for the UMediaPlateComponent class.
@@ -52,5 +53,9 @@ private:
 	 */
 	void OnMediaPathChanged(IDetailLayoutBuilder* DetailBuilder);
 
+	/**
+	 * Get the rate to use when we press the forward button.
+	 */
+	float GetForwardRate(UMediaPlayer* MediaPlayer) const;
 };
 
