MediaPlate: Play button resumes when paused.
#rnx

diff --git a/Engine/Plugins/Media/MediaPlate/Source/MediaPlateEditor/Private/MediaPlateCustomization.cpp b/Engine/Plugins/Media/MediaPlate/Source/MediaPlateEditor/Private/MediaPlateCustomization.cpp
--- a/Engine/Plugins/Media/MediaPlate/Source/MediaPlateEditor/Private/MediaPlateCustomization.cpp
+++ b/Engine/Plugins/Media/MediaPlate/Source/MediaPlateEditor/Private/MediaPlateCustomization.cpp
@@ -103,15 +103,26 @@ void FMediaPlateCustomization::CustomizeDetails(IDetailLayoutBuilder& DetailBuil
 								UMediaPlateComponent* MediaPlate = MediaPlatePtr.Get();
 								if (MediaPlate != nullptr)
 								{
-									// Tell the editor module that this media plate is playing.
-									FMediaPlateEditorModule* EditorModule = FModuleManager::LoadModulePtr<FMediaPlateEditorModule>("MediaPlateEditor");
-									if (EditorModule != nullptr)
+									// Is the player paused or fast forwarding/rewinding?
+									UMediaPlayer* MediaPlayer = MediaPlate->GetMediaPlayer();
+									if ((MediaPlayer != nullptr) &&
+										((MediaPlayer->IsPaused()) || 
+											(MediaPlayer->IsPlaying() && (MediaPlayer->GetRate() != 1.0f))))
 									{
-										EditorModule->MediaPlateStartedPlayback(MediaPlate);
+										MediaPlayer->Play();
 									}
+									else
+									{
+										// Tell the editor module that this media plate is playing.
+										FMediaPlateEditorModule* EditorModule = FModuleManager::LoadModulePtr<FMediaPlateEditorModule>("MediaPlateEditor");
+										if (EditorModule != nullptr)
+										{
+											EditorModule->MediaPlateStartedPlayback(MediaPlate);
+										}
 
-									// Play the media.
-									MediaPlate->Play();
+										// Play the media.
+										MediaPlate->Play();
+									}
 								}
 							}
 							return FReply::Handled();
