MediaPlate: Plugin enabled by default.
#rnx

diff --git a/Engine/Plugins/Media/MediaCompositing/MediaCompositing.uplugin b/Engine/Plugins/Media/MediaCompositing/MediaCompositing.uplugin
--- a/Engine/Plugins/Media/MediaCompositing/MediaCompositing.uplugin
+++ b/Engine/Plugins/Media/MediaCompositing/MediaCompositing.uplugin
@@ -32,6 +32,10 @@
 			"Name": "ImgMedia",
 			"TargetAllowList": [ "Editor" ],
 			"Enabled": true
+		},
+		{
+			"Name": "MediaPlate",
+			"Enabled": true
 		}
 	]
 }
diff --git a/Engine/Plugins/Media/MediaPlate/MediaPlate.uplugin b/Engine/Plugins/Media/MediaPlate/MediaPlate.uplugin
--- a/Engine/Plugins/Media/MediaPlate/MediaPlate.uplugin
+++ b/Engine/Plugins/Media/MediaPlate/MediaPlate.uplugin
@@ -11,11 +11,12 @@
 	"MarketplaceURL": "",
 	"SupportURL": "",
 	"CanContainContent": true,
-	"EnabledByDefault": false,
+	"EnabledByDefault": true,
 	"IsBetaVersion": true,
 	"IsExperimentalVersion": false,
 	"Installed": false,
-	"Modules": [
+	"Modules": 
+	[
 		{
 			"Name": "MediaPlate",
 			"Type": "Runtime",
@@ -26,5 +27,13 @@
 			"Type": "Editor",
 			"LoadingPhase": "Default"
 		}
+	],
+	"Plugins": 
+	[
+		{
+			"Name": "MediaPlayerEditor",
+			"TargetAllowList": [ "Editor" ],
+			"Enabled": true
+		}
 	]
 }
\ No newline at end of file
