MediaPlate: Added rewind and reverse buttons to customization.
#rnx

diff --git a/Engine/Plugins/Media/MediaPlate/Source/MediaPlateEditor/Private/MediaPlateCustomization.cpp b/Engine/Plugins/Media/MediaPlate/Source/MediaPlateEditor/Private/MediaPlateCustomization.cpp
--- a/Engine/Plugins/Media/MediaPlate/Source/MediaPlateEditor/Private/MediaPlateCustomization.cpp
+++ b/Engine/Plugins/Media/MediaPlate/Source/MediaPlateEditor/Private/MediaPlateCustomization.cpp
@@ -89,6 +89,66 @@ void FMediaPlateCustomization::CustomizeDetails(IDetailLayoutBuilder& DetailBuil
 		[
 			SNew(SHorizontalBox)
 
+			// Rewind button.
+			+ SHorizontalBox::Slot()
+				.VAlign(VAlign_Center)
+				.AutoWidth()
+				[
+					SNew(SButton)
+					.VAlign(VAlign_Center)
+					.OnClicked_Lambda([this]() -> FReply
+					{
+						for (TWeakObjectPtr<UMediaPlateComponent>& MediaPlatePtr : MediaPlatesList)
+						{
+							UMediaPlateComponent* MediaPlate = MediaPlatePtr.Get();
+							if (MediaPlate != nullptr)
+							{
+								UMediaPlayer* MediaPlayer = MediaPlate->GetMediaPlayer();
+								if (MediaPlayer != nullptr)
+								{
+									MediaPlayer->Rewind();
+								}
+							}
+						}
+						return FReply::Handled();
+					})
+					[
+						SNew(SImage)
+						.ColorAndOpacity(FSlateColor::UseForeground())
+					.Image(Style->GetBrush("MediaPlateEditor.RewindMedia.Small"))
+					]
+				]
+
+			// Reverse button.
+			+ SHorizontalBox::Slot()
+				.VAlign(VAlign_Center)
+				.AutoWidth()
+				[
+					SNew(SButton)
+						.VAlign(VAlign_Center)
+						.OnClicked_Lambda([this]() -> FReply
+						{
+							for (TWeakObjectPtr<UMediaPlateComponent>& MediaPlatePtr : MediaPlatesList)
+							{
+								UMediaPlateComponent* MediaPlate = MediaPlatePtr.Get();
+								if (MediaPlate != nullptr)
+								{
+									UMediaPlayer* MediaPlayer = MediaPlate->GetMediaPlayer();
+									if (MediaPlayer != nullptr)
+									{
+										MediaPlayer->SetRate(GetReverseRate(MediaPlayer));
+									}
+								}
+							}
+							return FReply::Handled();
+						})
+						[
+							SNew(SImage)
+								.ColorAndOpacity(FSlateColor::UseForeground())
+								.Image(Style->GetBrush("MediaPlateEditor.ReverseMedia.Small"))
+						]
+				]
+
 			// Play button.
 			+ SHorizontalBox::Slot()
 				.VAlign(VAlign_Center)
@@ -272,7 +332,6 @@ void FMediaPlateCustomization::OnMediaPathChanged(IDetailLayoutBuilder* DetailBu
 	DetailBuilder->ForceRefreshDetails();
 }
 
-
 float FMediaPlateCustomization::GetForwardRate(UMediaPlayer* MediaPlayer) const
 {
 	float Rate = MediaPlayer->GetRate();
@@ -285,5 +344,17 @@ float FMediaPlateCustomization::GetForwardRate(UMediaPlayer* MediaPlayer) const
 	return 2.0f * Rate;
 }
 
+float FMediaPlateCustomization::GetReverseRate(UMediaPlayer* MediaPlayer) const
+{
+	float Rate = MediaPlayer->GetRate();
+
+	if (Rate > -1.0f)
+	{
+		return -1.0f;
+	}
+
+	return 2.0f * Rate;
+}
+
 
 #undef LOCTEXT_NAMESPACE
diff --git a/Engine/Plugins/Media/MediaPlate/Source/MediaPlateEditor/Private/MediaPlateCustomization.h b/Engine/Plugins/Media/MediaPlate/Source/MediaPlateEditor/Private/MediaPlateCustomization.h
--- a/Engine/Plugins/Media/MediaPlate/Source/MediaPlateEditor/Private/MediaPlateCustomization.h
+++ b/Engine/Plugins/Media/MediaPlate/Source/MediaPlateEditor/Private/MediaPlateCustomization.h
@@ -57,5 +57,10 @@ private:
 	 * Get the rate to use when we press the forward button.
 	 */
 	float GetForwardRate(UMediaPlayer* MediaPlayer) const;
+
+	/**
+	 * Get the rate to use when we press the reverse button.
+	 */
+	float GetReverseRate(UMediaPlayer* MediaPlayer) const;
 };
 
