MediaPlate: Added bool to select between media path and media source.

#rnx

diff --git a/Engine/Plugins/Media/MediaPlate/Source/MediaPlate/Private/MediaPlateComponent.cpp b/Engine/Plugins/Media/MediaPlate/Source/MediaPlate/Private/MediaPlateComponent.cpp
--- a/Engine/Plugins/Media/MediaPlate/Source/MediaPlate/Private/MediaPlateComponent.cpp
+++ b/Engine/Plugins/Media/MediaPlate/Source/MediaPlate/Private/MediaPlateComponent.cpp
@@ -86,19 +86,21 @@ void UMediaPlateComponent::Play()
 	TObjectPtr<UMediaPlayer> MediaPlayer = GetMediaPlayer();
 	if (MediaPlayer != nullptr)
 	{
-		// Try and play the media path.
-		if (MediaPath.FilePath.IsEmpty() == false)
+		if (bUseMediaSource == false)
 		{
-			MediaPathMediaSource = UMediaSource::SpawnMediaSourceForString(MediaPath.FilePath);
-			if (MediaPathMediaSource != nullptr)
+			// Try and play the media path.
+			if (MediaPath.FilePath.IsEmpty() == false)
 			{
-				bIsPlaying = PlayMediaSource(MediaPathMediaSource);
+				MediaPathMediaSource = UMediaSource::SpawnMediaSourceForString(MediaPath.FilePath);
+				if (MediaPathMediaSource != nullptr)
+				{
+					bIsPlaying = PlayMediaSource(MediaPathMediaSource);
+				}
 			}
 		}
-
-		// If we did not get anything, try the media source.
-		if (bIsPlaying == false)
+		else
 		{
+			// Try the media source.
 			bIsPlaying = PlayMediaSource(MediaSource);
 		}
 	}
diff --git a/Engine/Plugins/Media/MediaPlate/Source/MediaPlate/Public/MediaPlateComponent.h b/Engine/Plugins/Media/MediaPlate/Source/MediaPlate/Public/MediaPlateComponent.h
--- a/Engine/Plugins/Media/MediaPlate/Source/MediaPlate/Public/MediaPlateComponent.h
+++ b/Engine/Plugins/Media/MediaPlate/Source/MediaPlate/Public/MediaPlateComponent.h
@@ -79,12 +79,16 @@ public:
 	UPROPERTY(EditAnywhere, AdvancedDisplay, Category = MediaPlate)
 	TObjectPtr<UMediaSoundComponent> SoundComponent;
 
-	/** URL (or file)  to play. This will take precedence over the MediaSource. */
+	/** Selects whether to use the media source or the media path. */
 	UPROPERTY(EditAnywhere, Category = MediaPlate)
+	bool bUseMediaSource;
+
+	/** URL (or file) to play. */
+	UPROPERTY(EditAnywhere, Category = MediaPlate, meta = (EditCondition = "!bUseMediaSource"))
 	FFilePath MediaPath;
 
-	/** What media to play. This will only be used if MediaPath is not valid. */
-	UPROPERTY(EditAnywhere, Category = MediaPlate)
+	/** What media to play. */
+	UPROPERTY(EditAnywhere, Category = MediaPlate, meta = (EditCondition = "bUseMediaSource"))
 	TObjectPtr<UMediaSource> MediaSource;
 
 	/** Enable smart caching for image sequences. */
diff --git a/Engine/Plugins/Media/MediaPlate/Source/MediaPlateEditor/Private/MediaPlateCustomization.cpp b/Engine/Plugins/Media/MediaPlate/Source/MediaPlateEditor/Private/MediaPlateCustomization.cpp
--- a/Engine/Plugins/Media/MediaPlate/Source/MediaPlateEditor/Private/MediaPlateCustomization.cpp
+++ b/Engine/Plugins/Media/MediaPlate/Source/MediaPlateEditor/Private/MediaPlateCustomization.cpp
@@ -53,35 +53,15 @@ void FMediaPlateCustomization::CustomizeDetails(IDetailLayoutBuilder& DetailBuil
 		}
 	}
 
-	// Get media source property.
-	MediaSourceProperty = DetailBuilder.GetProperty(GET_MEMBER_NAME_CHECKED(UMediaPlateComponent, MediaSource));
-	
-	// Get media path property.
-	TSharedRef<IPropertyHandle> Property = DetailBuilder.GetProperty(GET_MEMBER_NAME_CHECKED(UMediaPlateComponent, MediaPath));
+	// Get the bUseMediaSource property.
+	TSharedRef<IPropertyHandle> Property = DetailBuilder.GetProperty(GET_MEMBER_NAME_CHECKED(UMediaPlateComponent, bUseMediaSource));
 	if (Property->IsValidHandle())
 	{
-		MediaPathProperty = Property->GetChildHandle("FilePath");
-		if (MediaPathProperty->IsValidHandle())
-		{
-			if (MediaSourceProperty->IsValidHandle())
-			{
-				// Get a callback when this changes so we can hide the media source property.
-				FSimpleDelegate OnUrlChangedDelegate = FSimpleDelegate::CreateSP(this,
-					&FMediaPlateCustomization::OnMediaPathChanged, &DetailBuilder);
-				MediaPathProperty->SetOnPropertyValueChanged(OnUrlChangedDelegate);
-
-				// Is the media path being used?
-				FString MediaPath;
-				if (MediaPathProperty->GetValue(MediaPath))
-				{
-					if (MediaPath.IsEmpty() == false)
-					{
-						// Yes, so hide the media source.
-						MediaSourceProperty->MarkHiddenByCustomization();
-					}
-				}
-			}
-		}
+		
+		// Get a callback when this changes so we can stop the player.
+		FSimpleDelegate OnUseMediaSourceChangedDelegate = FSimpleDelegate::CreateSP(this,
+			&FMediaPlateCustomization::OnUseMediaSourceChanged, &DetailBuilder);
+		Property->SetOnPropertyValueChanged(OnUseMediaSourceChangedDelegate);
 	}
 
 	// Add media control buttons.
@@ -326,10 +306,17 @@ FReply FMediaPlateCustomization::OnOpenMediaPlate()
 	return FReply::Handled();
 }
 
-void FMediaPlateCustomization::OnMediaPathChanged(IDetailLayoutBuilder* DetailBuilder)
+void FMediaPlateCustomization::OnUseMediaSourceChanged(IDetailLayoutBuilder* DetailBuilder)
 {
-	// Refresh the layout so we can show/hide the media source.
-	DetailBuilder->ForceRefreshDetails();
+	// Stop the player if we change the media source.
+	for (TWeakObjectPtr<UMediaPlateComponent>& MediaPlatePtr : MediaPlatesList)
+	{
+		UMediaPlateComponent* MediaPlate = MediaPlatePtr.Get();
+		if (MediaPlate != nullptr)
+		{
+			MediaPlate->Stop();
+		}
+	}
 }
 
 float FMediaPlateCustomization::GetForwardRate(UMediaPlayer* MediaPlayer) const
diff --git a/Engine/Plugins/Media/MediaPlate/Source/MediaPlateEditor/Private/MediaPlateCustomization.h b/Engine/Plugins/Media/MediaPlate/Source/MediaPlateEditor/Private/MediaPlateCustomization.h
--- a/Engine/Plugins/Media/MediaPlate/Source/MediaPlateEditor/Private/MediaPlateCustomization.h
+++ b/Engine/Plugins/Media/MediaPlate/Source/MediaPlateEditor/Private/MediaPlateCustomization.h
@@ -38,20 +38,15 @@ private:
 	/** List of the media plates we are editing. */
 	TArray<TWeakObjectPtr<UMediaPlateComponent>> MediaPlatesList;
 
-	/** Stores our media source property. */
-	TSharedPtr<IPropertyHandle> MediaSourceProperty;
-	/** Stores our media path property. */
-	TSharedPtr<IPropertyHandle> MediaPathProperty;
-
 	/**
 	 * Called when the open media plate button is pressed.
 	 */
 	FReply OnOpenMediaPlate();
 
 	/**
-	 * Called when MediaPath changes.
+	 * Called when bUseMediaSource changes.
 	 */
-	void OnMediaPathChanged(IDetailLayoutBuilder* DetailBuilder);
+	void OnUseMediaSourceChanged(IDetailLayoutBuilder* DetailBuilder);
 
 	/**
 	 * Get the rate to use when we press the forward button.
