MediaPlate: Added media details tab.

#rnx

diff --git a/Engine/Plugins/Media/MediaPlate/Source/MediaPlateEditor/Private/Toolkits/MediaPlateEditorToolkit.cpp b/Engine/Plugins/Media/MediaPlate/Source/MediaPlateEditor/Private/Toolkits/MediaPlateEditorToolkit.cpp
--- a/Engine/Plugins/Media/MediaPlate/Source/MediaPlateEditor/Private/Toolkits/MediaPlateEditorToolkit.cpp
+++ b/Engine/Plugins/Media/MediaPlate/Source/MediaPlateEditor/Private/Toolkits/MediaPlateEditorToolkit.cpp
@@ -15,6 +15,7 @@
 #include "SlateOptMacros.h"
 #include "Widgets/Docking/SDockTab.h"
 #include "Widgets/SMediaPlateEditorDetails.h"
+#include "Widgets/SMediaPlateEditorMediaDetails.h"
 #include "Widgets/SMediaPlayerEditorViewer.h"
 
 #define LOCTEXT_NAMESPACE "FMediaPlateEditorToolkit"
@@ -23,6 +24,7 @@ namespace MediaPlateEditorToolkit
 {
 	static const FName AppIdentifier("MediaPlateEditorApp");
 	static const FName DetailsTabId("Details");
+	static const FName MediaDetailsTabId("MediaDetails");
 	static const FName ViewerTabId("Viewer");
 }
 
@@ -62,7 +64,7 @@ void FMediaPlateEditorToolkit::Initialize(UMediaPlateComponent* InMediaPlate, co
 	BindCommands();
 
 	// create tab layout
-	const TSharedRef<FTabManager::FLayout> Layout = FTabManager::NewLayout("Standalone_MediaPlateEditor_v1.1b")
+	const TSharedRef<FTabManager::FLayout> Layout = FTabManager::NewLayout("Standalone_MediaPlateEditor_v1.2")
 		->AddArea
 		(
 			FTabManager::NewPrimaryArea()
@@ -81,10 +83,23 @@ void FMediaPlateEditorToolkit::Initialize(UMediaPlateComponent* InMediaPlate, co
 						)	
 						->Split
 						(
-							// viewer
-							FTabManager::NewStack()
-							->AddTab(MediaPlateEditorToolkit::DetailsTabId, ETabState::OpenedTab)
+							FTabManager::NewSplitter()
+							->SetOrientation(Orient_Vertical)
 							->SetSizeCoefficient(0.4f)
+							->Split
+							(
+								// Media details tab.
+								FTabManager::NewStack()
+								->AddTab(MediaPlateEditorToolkit::MediaDetailsTabId, ETabState::OpenedTab)
+								->SetSizeCoefficient(0.2f)
+							)
+							->Split
+							(
+								// Details tab.
+								FTabManager::NewStack()
+								->AddTab(MediaPlateEditorToolkit::DetailsTabId, ETabState::OpenedTab)
+								->SetSizeCoefficient(0.8f)
+							)
 						)
 				)
 		);
@@ -140,6 +155,12 @@ void FMediaPlateEditorToolkit::RegisterTabSpawners(const TSharedRef<FTabManager>
 		.SetGroup(WorkspaceMenuCategoryRef)
 		.SetIcon(FSlateIcon(FEditorStyle::GetStyleSetName(), "LevelEditor.Tabs.Details"));
 
+	// Media details tab.
+	InTabManager->RegisterTabSpawner(MediaPlateEditorToolkit::MediaDetailsTabId, FOnSpawnTab::CreateSP(this, &FMediaPlateEditorToolkit::HandleTabManagerSpawnTab, MediaPlateEditorToolkit::MediaDetailsTabId))
+		.SetDisplayName(LOCTEXT("MediaDetailsTabName", "Media Details"))
+		.SetGroup(WorkspaceMenuCategoryRef)
+		.SetIcon(FSlateIcon(FEditorStyle::GetStyleSetName(), "LevelEditor.Tabs.Details"));
+
 	// Viewer tab.
 	InTabManager->RegisterTabSpawner(MediaPlateEditorToolkit::ViewerTabId, FOnSpawnTab::CreateSP(this, &FMediaPlateEditorToolkit::HandleTabManagerSpawnTab, MediaPlateEditorToolkit::ViewerTabId))
 		.SetDisplayName(LOCTEXT("PlayerTabName", "Player"))
@@ -152,6 +173,7 @@ void FMediaPlateEditorToolkit::UnregisterTabSpawners(const TSharedRef<class FTab
 	FAssetEditorToolkit::UnregisterTabSpawners(InTabManager);
 
 	InTabManager->UnregisterTabSpawner(MediaPlateEditorToolkit::ViewerTabId);
+	InTabManager->UnregisterTabSpawner(MediaPlateEditorToolkit::MediaDetailsTabId);
 	InTabManager->UnregisterTabSpawner(MediaPlateEditorToolkit::DetailsTabId);
 }
 
@@ -356,6 +378,10 @@ TSharedRef<SDockTab> FMediaPlateEditorToolkit::HandleTabManagerSpawnTab(const FS
 	{
 		TabWidget = SNew(SMediaPlateEditorDetails, *MediaPlate, Style);
 	}
+	else if (TabIdentifier == MediaPlateEditorToolkit::MediaDetailsTabId)
+	{
+		TabWidget = SNew(SMediaPlateEditorMediaDetails, *MediaPlate, Style);
+	}
 	else if (TabIdentifier == MediaPlateEditorToolkit::ViewerTabId)
 	{
 		TabWidget = SNew(SMediaPlayerEditorViewer, *MediaPlayer, MediaTexture, Style, false);
diff --git a/Engine/Plugins/Media/MediaPlate/Source/MediaPlateEditor/Private/Widgets/SMediaPlateEditorMediaDetails.cpp b/Engine/Plugins/Media/MediaPlate/Source/MediaPlateEditor/Private/Widgets/SMediaPlateEditorMediaDetails.cpp
new file mode 100644
--- /dev/null
+++ b/Engine/Plugins/Media/MediaPlate/Source/MediaPlateEditor/Private/Widgets/SMediaPlateEditorMediaDetails.cpp
@@ -0,0 +1,167 @@
+// Copyright Epic Games, Inc. All Rights Reserved.
+
+#include "Widgets/SMediaPlateEditorMediaDetails.h"
+#include "IDetailsView.h"
+#include "MediaPlateComponent.h"
+#include "MediaPlayer.h"
+#include "MediaTexture.h"
+#include "Modules/ModuleManager.h"
+#include "PropertyEditorModule.h"
+#include "Widgets/Text/STextBlock.h"
+
+#define LOCTEXT_NAMESPACE "SMediaPlateEditorMediaDetails"
+
+/* SMediaPlateEditorMediaDetails interface
+ *****************************************************************************/
+
+void SMediaPlateEditorMediaDetails::Construct(const FArguments& InArgs, UMediaPlateComponent& InMediaPlate, const TSharedRef<ISlateStyle>& InStyle)
+{
+	MediaPlate = &InMediaPlate;
+
+	ChildSlot
+		[
+			SNew(SHorizontalBox)
+
+			// Left side.
+			+ SHorizontalBox::Slot()
+				[
+					SNew(SVerticalBox)
+
+					// Resolution.
+					+ SVerticalBox::Slot()
+						.AutoHeight()
+						.VAlign(VAlign_Center)
+						.Padding(4.0f)
+						[
+							SAssignNew(ResolutionText, STextBlock)
+						]
+
+					// Frame rate.
+					+ SVerticalBox::Slot()
+						.AutoHeight()
+						.VAlign(VAlign_Center)
+						.Padding(4.0f)
+						[
+							SAssignNew(FrameRateText, STextBlock)
+						]
+
+					// Resource size.
+					+ SVerticalBox::Slot()
+						.AutoHeight()
+						.VAlign(VAlign_Center)
+						.Padding(4.0f)
+						[
+							SAssignNew(ResourceSizeText, STextBlock)
+						]
+				]
+
+			// Right side.
+			+ SHorizontalBox::Slot()
+				[
+					SNew(SVerticalBox)
+
+					// Method.
+					+ SVerticalBox::Slot()
+						.AutoHeight()
+						.VAlign(VAlign_Center)
+						.Padding(4.0f)
+						[
+							SAssignNew(MethodText, STextBlock)
+						]
+
+					// Format.
+					+ SVerticalBox::Slot()
+						.AutoHeight()
+						.VAlign(VAlign_Center)
+						.Padding(4.0f)
+						[
+							SAssignNew(FormatText, STextBlock)
+						]
+
+					// LOD bias.
+					+ SVerticalBox::Slot()
+						.AutoHeight()
+						.VAlign(VAlign_Center)
+						.Padding(4.0f)
+						[
+							SAssignNew(LODBiasText, STextBlock)
+						]
+
+					// Num mips.
+					+ SVerticalBox::Slot()
+						.AutoHeight()
+						.VAlign(VAlign_Center)
+						.Padding(4.0f)
+						[
+							SAssignNew(NumMipsText, STextBlock)
+						]
+				]
+
+		];
+
+	UpdateDetails();
+}
+
+
+void SMediaPlateEditorMediaDetails::Tick(const FGeometry& AllottedGeometry, const double InCurrentTime, const float InDeltaTime)
+{
+	// Call parent.
+	SCompoundWidget::Tick(AllottedGeometry, InCurrentTime, InDeltaTime);
+
+	UpdateDetails();
+}
+
+void SMediaPlateEditorMediaDetails::UpdateDetails()
+{
+	FString Format;
+	float FrameRate = 0.0f;
+	int32 LODBias = 0;
+	FText Method;
+	int32 NumMips = 0;
+	int64 ResourceSize = 0;
+	int32 SurfaceWidth = 0;
+	int32 SurfaceHeight = 0;
+
+	if (MediaPlate.IsValid())
+	{
+		// Get player info.
+		UMediaPlayer* MediaPlayer = MediaPlate->GetMediaPlayer();
+		if (MediaPlayer != nullptr)
+		{
+			FrameRate = MediaPlayer->GetVideoTrackFrameRate(INDEX_NONE, INDEX_NONE);
+			Format = MediaPlayer->GetVideoTrackType(INDEX_NONE, INDEX_NONE);
+		}
+
+		// Get texture info,
+		UMediaTexture* MediaTexture = MediaPlate->GetMediaTexture();
+		if (MediaTexture != nullptr)
+		{
+			LODBias = MediaTexture->GetCachedLODBias();
+			Method = MediaTexture->IsCurrentlyVirtualTextured() ?
+				LOCTEXT("MethodVirtualStreamed", "Virtual Streamed")
+				: (!MediaTexture->IsStreamable() ? LOCTEXT("QuickInfo_MethodNotStreamed", "Not Streamed")
+					: LOCTEXT("MethodStreamed", "Streamed"));
+			NumMips = MediaTexture->GetTextureNumMips();
+			ResourceSize = (MediaTexture->GetResourceSizeBytes(EResourceSizeMode::Exclusive) + 512) / 1024;
+			SurfaceWidth = MediaTexture->GetSurfaceWidth();
+			SurfaceHeight = MediaTexture->GetSurfaceHeight();
+		}
+	}
+
+	// Update text.
+	FormatText->SetText(FText::Format(LOCTEXT("Format", "Format: {0}"),
+		FText::FromString(Format)));
+	FrameRateText->SetText(FText::Format(LOCTEXT("FrameRate", "Frame Rate: {0}"), 
+		FText::AsNumber(FrameRate)));
+	LODBiasText->SetText(FText::Format(LOCTEXT("LODBias", "Combined LOD Bias: {0}"),
+		FText::AsNumber(LODBias)));
+	MethodText->SetText(FText::Format(LOCTEXT("Method", "Method: {0}"), Method));
+	NumMipsText->SetText(FText::Format(LOCTEXT("NumberOfMips", "Number Of Mips: {0}"),
+		FText::AsNumber(NumMips)));
+	ResolutionText->SetText(FText::Format(LOCTEXT("Resolution", "Resolution: {0}x{1}"),
+		FText::AsNumber(SurfaceWidth), FText::AsNumber(SurfaceHeight)));
+	ResourceSizeText->SetText(FText::Format(LOCTEXT("ResourceSize", "Resource Size: {0} KB"),
+		FText::AsNumber(ResourceSize)));
+}
+
+#undef LOCTEXT_NAMESPACE
diff --git a/Engine/Plugins/Media/MediaPlate/Source/MediaPlateEditor/Private/Widgets/SMediaPlateEditorMediaDetails.h b/Engine/Plugins/Media/MediaPlate/Source/MediaPlateEditor/Private/Widgets/SMediaPlateEditorMediaDetails.h
new file mode 100644
--- /dev/null
+++ b/Engine/Plugins/Media/MediaPlate/Source/MediaPlateEditor/Private/Widgets/SMediaPlateEditorMediaDetails.h
@@ -0,0 +1,56 @@
+// Copyright Epic Games, Inc. All Rights Reserved.
+
+#pragma once
+
+#include "CoreMinimal.h"
+#include "Styling/ISlateStyle.h"
+#include "Widgets/DeclarativeSyntaxSupport.h"
+#include "Widgets/SCompoundWidget.h"
+
+class STextBlock;
+class UMediaPlateComponent;
+
+/**
+ * Implements the details panel of the MediaPlate asset editor.
+ */
+class SMediaPlateEditorMediaDetails
+	: public SCompoundWidget
+{
+public:
+
+	SLATE_BEGIN_ARGS(SMediaPlateEditorMediaDetails) { }
+	SLATE_END_ARGS()
+
+public:
+
+	/**
+	 * Construct this widget.
+	 *
+	 * @param InArgs		The declaration data for this widget.
+	 * @param InMediaPlate	The MediaPlate to show the details for.
+	 * @param InStyleSet	The style set to use.
+	 */
+	void Construct(const FArguments& InArgs, UMediaPlateComponent& InMediaPlate, const TSharedRef<ISlateStyle>& InStyle);
+
+	//~ SWidget interface
+	virtual void Tick(const FGeometry& AllottedGeometry, const double InCurrentTime, const float InDeltaTime) override;
+
+private:
+
+	/**
+	 * Updates our widgets to reflect the current state.
+	 */
+	void UpdateDetails();
+
+	/** Pointer to the MediaPlate that is being viewed. */
+	TWeakObjectPtr<UMediaPlateComponent> MediaPlate;
+
+	/** Our widgets. */
+	TSharedPtr<STextBlock> FormatText;
+	TSharedPtr<STextBlock> FrameRateText;
+	TSharedPtr<STextBlock> LODBiasText;
+	TSharedPtr<STextBlock> MethodText;
+	TSharedPtr<STextBlock> NumMipsText;
+	TSharedPtr<STextBlock> ResolutionText;
+	TSharedPtr<STextBlock> ResourceSizeText;
+};
