MediaPlate: Added MediaPlateTrack to remove media plate dependencies in MediaTrackEditor.
#rnx

diff --git a/Engine/Plugins/Media/MediaCompositing/MediaCompositing.uplugin b/Engine/Plugins/Media/MediaCompositing/MediaCompositing.uplugin
--- a/Engine/Plugins/Media/MediaCompositing/MediaCompositing.uplugin
+++ b/Engine/Plugins/Media/MediaCompositing/MediaCompositing.uplugin
@@ -32,10 +32,6 @@
 			"Name": "ImgMedia",
 			"TargetAllowList": [ "Editor" ],
 			"Enabled": true
-		},
-		{
-			"Name": "MediaPlate",
-			"Enabled": true
 		}
 	]
 }
diff --git a/Engine/Plugins/Media/MediaCompositing/Source/MediaCompositingEditor/MediaCompositingEditor.Build.cs b/Engine/Plugins/Media/MediaCompositing/Source/MediaCompositingEditor/MediaCompositingEditor.Build.cs
--- a/Engine/Plugins/Media/MediaCompositing/Source/MediaCompositingEditor/MediaCompositingEditor.Build.cs
+++ b/Engine/Plugins/Media/MediaCompositing/Source/MediaCompositingEditor/MediaCompositingEditor.Build.cs
@@ -44,15 +44,6 @@ namespace UnrealBuildTool.Rules
 					"MediaCompositingEditor/Private/Sequencer",
 					"MediaCompositingEditor/Private/Shared",
 				});
-
-			DynamicallyLoadedModuleNames.AddRange(
-				new string[] {
-					"MediaPlate",
-				});
-			PrivateIncludePathModuleNames.AddRange(
-				new string[] {
-					"MediaPlate",
-				});
 		}
 	}
 }
diff --git a/Engine/Plugins/Media/MediaCompositing/Source/MediaCompositingEditor/Private/Sequencer/MediaTrackEditor.cpp b/Engine/Plugins/Media/MediaCompositing/Source/MediaCompositingEditor/Private/Sequencer/MediaTrackEditor.cpp
--- a/Engine/Plugins/Media/MediaCompositing/Source/MediaCompositingEditor/Private/Sequencer/MediaTrackEditor.cpp
+++ b/Engine/Plugins/Media/MediaCompositing/Source/MediaCompositingEditor/Private/Sequencer/MediaTrackEditor.cpp
@@ -11,8 +11,6 @@
 #include "ISequencerObjectChangeListener.h"
 #include "MediaSource.h"
 #include "MediaTexture.h"
-#include "MediaPlate.h"
-#include "MediaPlateModule.h"
 #include "MovieSceneBinding.h"
 #include "MovieSceneMediaSection.h"
 #include "MovieSceneMediaTrack.h"
@@ -78,27 +76,6 @@ void FMediaTrackEditor::BuildAddTrackMenu(FMenuBuilder& MenuBuilder)
 	);
 }
 
-void FMediaTrackEditor::BuildObjectBindingTrackMenu(FMenuBuilder& MenuBuilder, const TArray<FGuid>& ObjectBindings, const UClass* ObjectClass)
-{
-	// Is this a media plate?
-	if (ObjectClass != nullptr)
-	{
-		FMediaPlateModule* MediaPlateModule = FModuleManager::LoadModulePtr<FMediaPlateModule>("MediaPlate");
-		if (MediaPlateModule != nullptr)
-		{
-			if (ObjectClass->IsChildOf(MediaPlateModule->GetAMediaPlateClass()))
-			{
-				MenuBuilder.AddMenuEntry(
-					LOCTEXT("AddTrack", "Media"),
-					LOCTEXT("AddAttachedTooltip", "Adds a media track attached to the object."),
-					FSlateIcon(),
-					FUIAction(
-						FExecuteAction::CreateRaw(this, &FMediaTrackEditor::HandleAddMediaTrackToObjectBindingMenuEntryExecute, ObjectBindings)));
-			}
-		}
-	}
-
-}
 
 TSharedPtr<SWidget> FMediaTrackEditor::BuildOutlinerEditWidget(const FGuid& ObjectBinding, UMovieSceneTrack* Track, const FBuildEditWidgetParams& Params)
 {
@@ -334,38 +311,6 @@ void FMediaTrackEditor::HandleAddMediaTrackMenuEntryExecute()
 	}
 }
 
-void FMediaTrackEditor::HandleAddMediaTrackToObjectBindingMenuEntryExecute(TArray<FGuid> InObjectBindingIDs)
-{
-	UMovieScene* FocusedMovieScene = GetFocusedMovieScene();
-
-	if (FocusedMovieScene == nullptr)
-	{
-		return;
-	}
 
-	if (FocusedMovieScene->IsReadOnly())
-	{
-		return;
-	}
-
-	const FScopedTransaction Transaction(NSLOCTEXT("Sequencer", "AddMediaTrack_Transaction", "Add Media Track"));
-	FocusedMovieScene->Modify();
-
-	// Loop through all objects.
-	for (FGuid InObjectBindingID : InObjectBindingIDs)
-	{
-		if (InObjectBindingID.IsValid())
-		{
-			// Add media track.
-			UMovieSceneMediaTrack* NewObjectTrack = FocusedMovieScene->AddTrack<UMovieSceneMediaTrack>(InObjectBindingID);
-			NewObjectTrack->SetDisplayName(LOCTEXT("MediaTrackName", "Media"));
-			
-			if (GetSequencer().IsValid())
-			{
-				GetSequencer()->OnAddTrack(NewObjectTrack, InObjectBindingID);
-			}
-		}
-	}
-}
 
 #undef LOCTEXT_NAMESPACE
diff --git a/Engine/Plugins/Media/MediaCompositing/Source/MediaCompositingEditor/Public/Sequencer/MediaTrackEditor.h b/Engine/Plugins/Media/MediaCompositing/Source/MediaCompositingEditor/Public/Sequencer/MediaTrackEditor.h
--- a/Engine/Plugins/Media/MediaCompositing/Source/MediaCompositingEditor/Public/Sequencer/MediaTrackEditor.h
+++ b/Engine/Plugins/Media/MediaCompositing/Source/MediaCompositingEditor/Public/Sequencer/MediaTrackEditor.h
@@ -57,7 +57,6 @@ public:
 
 	virtual UMovieSceneTrack* AddTrack(UMovieScene* FocusedMovieScene, const FGuid& ObjectHandle, TSubclassOf<class UMovieSceneTrack> TrackClass, FName UniqueTypeName) override;
 	virtual void BuildAddTrackMenu(FMenuBuilder& MenuBuilder) override;
-	virtual void BuildObjectBindingTrackMenu(FMenuBuilder& MenuBuilder, const TArray<FGuid>& ObjectBindings, const UClass* ObjectClass) override;
 	virtual TSharedPtr<SWidget> BuildOutlinerEditWidget(const FGuid& ObjectBinding, UMovieSceneTrack* Track, const FBuildEditWidgetParams& Params) override;
 	virtual bool HandleAssetAdded(UObject* Asset, const FGuid& TargetObjectGuid) override;
 	virtual TSharedRef<ISequencerSection> MakeSectionInterface(UMovieSceneSection& SectionObject, UMovieSceneTrack& Track, FGuid ObjectBinding) override;
@@ -83,8 +82,6 @@ private:
 
 	/** Callback for executing the "Add Media Track" menu entry. */
 	void HandleAddMediaTrackMenuEntryExecute();
-	/** Callback for adding a media track to an object binding. */
-	void HandleAddMediaTrackToObjectBindingMenuEntryExecute(TArray<FGuid> InObjectBindingID);
 
 private:
 
diff --git a/Engine/Plugins/Media/MediaPlate/MediaPlate.uplugin b/Engine/Plugins/Media/MediaPlate/MediaPlate.uplugin
--- a/Engine/Plugins/Media/MediaPlate/MediaPlate.uplugin
+++ b/Engine/Plugins/Media/MediaPlate/MediaPlate.uplugin
@@ -30,6 +30,11 @@
 	],
 	"Plugins": 
 	[
+		{
+			"Name": "MediaCompositing",
+			"TargetAllowList": [ "Editor" ],
+			"Enabled": true
+		},
 		{
 			"Name": "MediaPlayerEditor",
 			"TargetAllowList": [ "Editor" ],
diff --git a/Engine/Plugins/Media/MediaPlate/Source/MediaPlateEditor/MediaPlateEditor.Build.cs b/Engine/Plugins/Media/MediaPlate/Source/MediaPlateEditor/MediaPlateEditor.Build.cs
--- a/Engine/Plugins/Media/MediaPlate/Source/MediaPlateEditor/MediaPlateEditor.Build.cs
+++ b/Engine/Plugins/Media/MediaPlate/Source/MediaPlateEditor/MediaPlateEditor.Build.cs
@@ -23,8 +23,12 @@ public class MediaPlateEditor : ModuleRules
 				"Engine",
 				"InputCore",
 				"MediaAssets",
+				"MediaCompositing",
+				"MediaCompositingEditor",
 				"MediaPlate",
 				"MediaPlayerEditor",
+				"MovieScene",
+				"Sequencer",
 				"Slate",
 				"SlateCore",
 				"UnrealEd",
diff --git a/Engine/Plugins/Media/MediaPlate/Source/MediaPlateEditor/Private/MediaPlateEditorModule.cpp b/Engine/Plugins/Media/MediaPlate/Source/MediaPlateEditor/Private/MediaPlateEditorModule.cpp
--- a/Engine/Plugins/Media/MediaPlate/Source/MediaPlateEditor/Private/MediaPlateEditorModule.cpp
+++ b/Engine/Plugins/Media/MediaPlate/Source/MediaPlateEditor/Private/MediaPlateEditorModule.cpp
@@ -4,6 +4,7 @@
 
 #include "AssetTools/MediaPlateActions.h"
 #include "Editor.h"
+#include "ISequencerModule.h"
 #include "LevelEditorViewport.h"
 #include "MediaPlateComponent.h"
 #include "MediaPlateCustomization.h"
@@ -12,6 +13,7 @@
 #include "MediaSoundComponent.h"
 #include "Models/MediaPlateEditorCommands.h"
 #include "PropertyEditorModule.h"
+#include "Sequencer/MediaPlateTrackEditor.h"
 
 DEFINE_LOG_CATEGORY(LogMediaPlateEditor);
 
@@ -29,10 +31,19 @@ void FMediaPlateEditorModule::StartupModule()
 	FPropertyEditorModule& PropertyModule = FModuleManager::LoadModuleChecked<FPropertyEditorModule>("PropertyEditor");
 	PropertyModule.RegisterCustomClassLayout(MediaPlateName,
 		FOnGetDetailCustomizationInstance::CreateStatic(&FMediaPlateCustomization::MakeInstance));
+
+	ISequencerModule& SequencerModule = FModuleManager::Get().LoadModuleChecked<ISequencerModule>("Sequencer");
+	TrackEditorBindingHandle = SequencerModule.RegisterPropertyTrackEditor<FMediaPlateTrackEditor>();
 }
 
 void FMediaPlateEditorModule::ShutdownModule()
 {
+	ISequencerModule* SequencerModulePtr = FModuleManager::Get().GetModulePtr<ISequencerModule>("Sequencer");
+	if (SequencerModulePtr)
+	{
+		SequencerModulePtr->UnRegisterTrackEditor(TrackEditorBindingHandle);
+	}
+
 	UnregisterAssetTools();
 
 	// Unregister customizations.
diff --git a/Engine/Plugins/Media/MediaPlate/Source/MediaPlateEditor/Private/Sequencer/MediaPlateTrackEditor.cpp b/Engine/Plugins/Media/MediaPlate/Source/MediaPlateEditor/Private/Sequencer/MediaPlateTrackEditor.cpp
new file mode 100644
--- /dev/null
+++ b/Engine/Plugins/Media/MediaPlate/Source/MediaPlateEditor/Private/Sequencer/MediaPlateTrackEditor.cpp
@@ -0,0 +1,132 @@
+// Copyright Epic Games, Inc. All Rights Reserved.
+
+#include "Sequencer/MediaPlateTrackEditor.h"
+
+#include "Framework/MultiBox/MultiBoxBuilder.h"
+#include "ISequencer.h"
+#include "MediaTexture.h"
+#include "MediaPlate.h"
+#include "MediaPlateComponent.h"
+#include "MovieSceneMediaTrack.h"
+
+#define LOCTEXT_NAMESPACE "FMediaPlateTrackEditor"
+
+/* FMediaTrackEditor static functions
+ *****************************************************************************/
+
+TArray<FAnimatedPropertyKey, TInlineAllocator<1>> FMediaPlateTrackEditor::GetAnimatedPropertyTypes()
+{
+	return TArray<FAnimatedPropertyKey, TInlineAllocator<1>>({ FAnimatedPropertyKey::FromObjectType(UMediaTexture::StaticClass()) });
+}
+
+/* FMediaTrackEditor structors
+ *****************************************************************************/
+
+FMediaPlateTrackEditor::FMediaPlateTrackEditor(TSharedRef<ISequencer> InSequencer)
+	: FMovieSceneTrackEditor(InSequencer)
+{
+	OnActorAddedToSequencerHandle = InSequencer->OnActorAddedToSequencer().AddRaw(this, &FMediaPlateTrackEditor::HandleActorAdded);
+}
+
+FMediaPlateTrackEditor::~FMediaPlateTrackEditor()
+{
+}
+
+void FMediaPlateTrackEditor::BuildObjectBindingTrackMenu(FMenuBuilder& MenuBuilder, const TArray<FGuid>& ObjectBindings, const UClass* ObjectClass)
+{
+	// Is this a media plate?
+	if (ObjectClass != nullptr)
+	{
+		if (ObjectClass->IsChildOf(AMediaPlate::StaticClass()))
+		{
+			MenuBuilder.AddMenuEntry(
+				LOCTEXT("AddTrack", "Media"),
+				LOCTEXT("AddAttachedTooltip", "Adds a media track attached to the object."),
+				FSlateIcon(),
+				FUIAction(
+					FExecuteAction::CreateRaw(this, &FMediaPlateTrackEditor::HandleAddMediaTrackToObjectBindingMenuEntryExecute, ObjectBindings)));
+		}
+	}
+}
+
+bool FMediaPlateTrackEditor::SupportsType(TSubclassOf<UMovieSceneTrack> TrackClass) const
+{
+	return false;
+}
+
+void FMediaPlateTrackEditor::HandleAddMediaTrackToObjectBindingMenuEntryExecute(TArray<FGuid> InObjectBindingIDs)
+{
+	UMovieScene* FocusedMovieScene = GetFocusedMovieScene();
+
+	if (FocusedMovieScene == nullptr)
+	{
+		return;
+	}
+
+	if (FocusedMovieScene->IsReadOnly())
+	{
+		return;
+	}
+
+	const FScopedTransaction Transaction(NSLOCTEXT("Sequencer", "AddMediaTrack_Transaction", "Add Media Track"));
+	FocusedMovieScene->Modify();
+
+	// Loop through all objects.
+	for (FGuid InObjectBindingID : InObjectBindingIDs)
+	{
+		if (InObjectBindingID.IsValid())
+		{
+			// Add media track.
+			UMovieSceneMediaTrack* NewObjectTrack = FocusedMovieScene->AddTrack<UMovieSceneMediaTrack>(InObjectBindingID);
+			NewObjectTrack->SetDisplayName(LOCTEXT("MediaTrackName", "Media"));
+
+			if (GetSequencer().IsValid())
+			{
+				GetSequencer()->OnAddTrack(NewObjectTrack, InObjectBindingID);
+			}
+		}
+	}
+}
+
+void FMediaPlateTrackEditor::HandleActorAdded(AActor* Actor, FGuid TargetObjectGuid)
+{
+	if (Actor)
+	{
+		if (UMediaPlateComponent* MediaPlateComponent = Actor->FindComponentByClass<UMediaPlateComponent>())
+		{
+			AddTrackForComponent(MediaPlateComponent);
+		}
+	}
+}
+
+void FMediaPlateTrackEditor::AddTrackForComponent(UMediaPlateComponent* Component)
+{
+	// Get object.
+	UObject* Object = Component->GetOwner();
+	FFindOrCreateHandleResult HandleResult = FindOrCreateHandleToObject(Object);
+	FGuid ObjectHandle = HandleResult.Handle;
+
+	// Add media track.
+	FFindOrCreateTrackResult TrackResult = FindOrCreateTrackForObject(ObjectHandle, UMovieSceneMediaTrack::StaticClass());
+	UMovieSceneTrack* Track = TrackResult.Track;
+	if (Track != nullptr)
+	{
+		UMovieSceneMediaTrack* MediaTrack = Cast<UMovieSceneMediaTrack>(Track);
+		MediaTrack->SetDisplayName(LOCTEXT("MediaTrackName", "Media"));
+	}
+}
+
+void FMediaPlateTrackEditor::OnRelease()
+{
+	if (GetSequencer().IsValid())
+	{
+		if (OnActorAddedToSequencerHandle.IsValid())
+		{
+			GetSequencer()->OnActorAddedToSequencer().Remove(OnActorAddedToSequencerHandle);
+		}
+	}
+
+	FMovieSceneTrackEditor::OnRelease();
+}
+
+#undef LOCTEXT_NAMESPACE
diff --git a/Engine/Plugins/Media/MediaPlate/Source/MediaPlateEditor/Private/Sequencer/MediaPlateTrackEditor.h b/Engine/Plugins/Media/MediaPlate/Source/MediaPlateEditor/Private/Sequencer/MediaPlateTrackEditor.h
new file mode 100644
--- /dev/null
+++ b/Engine/Plugins/Media/MediaPlate/Source/MediaPlateEditor/Private/Sequencer/MediaPlateTrackEditor.h
@@ -0,0 +1,75 @@
+// Copyright Epic Games, Inc. All Rights Reserved.
+
+#pragma once
+
+#include "AnimatedPropertyKey.h"
+#include "MovieSceneTrackEditor.h"
+#include "Templates/SharedPointer.h"
+
+class AActor;
+class ISequencer;
+class UMediaPlateComponent;
+
+/**
+ * Track editor for media plate.
+ */
+class FMediaPlateTrackEditor
+	: public FMovieSceneTrackEditor
+{
+public:
+
+	/**
+	 * Create a new media plate track editor instance.
+	 *
+	 * @param OwningSequencer The sequencer object that will own the track editor.
+	 * @return The new track editor.
+	 */
+	static TSharedRef<ISequencerTrackEditor>  CreateTrackEditor(TSharedRef<ISequencer> OwningSequencer)
+	{
+		return MakeShared<FMediaPlateTrackEditor>(OwningSequencer);
+	}
+
+	/**
+	 * Get the list of all property types that this track editor animates.
+	 *
+	 * @return List of animated properties.
+	 */
+	static TArray<FAnimatedPropertyKey, TInlineAllocator<1>> GetAnimatedPropertyTypes();
+
+public:
+
+	/**
+	 * Create and initialize a new instance.
+	 *
+	 * @param InSequencer The sequencer object that owns this track editor.
+	 */
+	FMediaPlateTrackEditor(TSharedRef<ISequencer> InSequencer);
+
+	/** Virtual destructor. */
+	virtual ~FMediaPlateTrackEditor();
+
+public:
+
+	//~ FMovieSceneTrackEditor interface
+	virtual void BuildObjectBindingTrackMenu(FMenuBuilder& MenuBuilder, const TArray<FGuid>& ObjectBindings, const UClass* ObjectClass) override;
+	virtual bool SupportsType(TSubclassOf<UMovieSceneTrack> TrackClass) const override;
+	//~ ISequencerTrackEditor interface
+	virtual void OnRelease() override;
+
+private:
+	/**
+	 * Callback for adding a media track to an object binding.
+	 */
+	void HandleAddMediaTrackToObjectBindingMenuEntryExecute(TArray<FGuid> InObjectBindingID);
+	/**
+	 * Callback when an actor is added.
+	 **/
+	void HandleActorAdded(AActor* Actor, FGuid TargetObjectGuid);
+	/**
+	 * Adds a track for this component.
+	 */
+	void AddTrackForComponent(UMediaPlateComponent* Component);
+
+	/** Handle to our delegate. */
+	FDelegateHandle OnActorAddedToSequencerHandle;
+};
diff --git a/Engine/Plugins/Media/MediaPlate/Source/MediaPlateEditor/Public/MediaPlateEditorModule.h b/Engine/Plugins/Media/MediaPlate/Source/MediaPlateEditor/Public/MediaPlateEditorModule.h
--- a/Engine/Plugins/Media/MediaPlate/Source/MediaPlateEditor/Public/MediaPlateEditorModule.h
+++ b/Engine/Plugins/Media/MediaPlate/Source/MediaPlateEditor/Public/MediaPlateEditorModule.h
@@ -49,6 +49,8 @@ private:
 	TSharedPtr<ISlateStyle> Style;
 	/** The collection of registered asset type actions. */
 	TArray<TSharedRef<IAssetTypeActions>> RegisteredAssetTypeActions;
+	/** Handle for our track editor. */
+	FDelegateHandle TrackEditorBindingHandle;
 
 	/**
 	 * Registers all of our asset tools.
