MediaPlate: Media path is now just an editor construct that modifies the media source.

#rnx

diff --git a/Engine/Plugins/Media/ImgMedia/Source/ImgMedia/Private/ImgMediaModule.cpp b/Engine/Plugins/Media/ImgMedia/Source/ImgMedia/Private/ImgMediaModule.cpp
--- a/Engine/Plugins/Media/ImgMedia/Source/ImgMedia/Private/ImgMediaModule.cpp
+++ b/Engine/Plugins/Media/ImgMedia/Source/ImgMedia/Private/ImgMediaModule.cpp
@@ -180,12 +180,12 @@ private:
 	 * Creates a media source for MediaPath.
 	 *
 	 * @param	MediaPath		File path to the media.
+	 * @param	Outer			Outer to use for this object.
 	 */
-	static UMediaSource* SpawnMediaSourceForString(const FString& MediaPath)
+	static UMediaSource* SpawnMediaSourceForString(const FString& MediaPath, UObject* Outer)
 	{
 		TObjectPtr<UImgMediaSource> MediaSource = 
-			NewObject<UImgMediaSource>(GetTransientPackage(), NAME_None,
-				RF_Transactional | RF_Transient);
+			NewObject<UImgMediaSource>(Outer, NAME_None, RF_Transactional);
 		MediaSource->SetSequencePath(MediaPath);
 
 		return MediaSource;
diff --git a/Engine/Plugins/Media/MediaPlate/Source/MediaPlate/Private/MediaPlateComponent.cpp b/Engine/Plugins/Media/MediaPlate/Source/MediaPlate/Private/MediaPlateComponent.cpp
--- a/Engine/Plugins/Media/MediaPlate/Source/MediaPlate/Private/MediaPlateComponent.cpp
+++ b/Engine/Plugins/Media/MediaPlate/Source/MediaPlate/Private/MediaPlateComponent.cpp
@@ -86,23 +86,7 @@ void UMediaPlateComponent::Play()
 	TObjectPtr<UMediaPlayer> MediaPlayer = GetMediaPlayer();
 	if (MediaPlayer != nullptr)
 	{
-		if (bUseMediaSource == false)
-		{
-			// Try and play the media path.
-			if (MediaPath.FilePath.IsEmpty() == false)
-			{
-				MediaPathMediaSource = UMediaSource::SpawnMediaSourceForString(MediaPath.FilePath);
-				if (MediaPathMediaSource != nullptr)
-				{
-					bIsPlaying = PlayMediaSource(MediaPathMediaSource);
-				}
-			}
-		}
-		else
-		{
-			// Try the media source.
-			bIsPlaying = PlayMediaSource(MediaSource);
-		}
+		bIsPlaying = PlayMediaSource(MediaSource);
 	}
 
 	// Did anything play?
diff --git a/Engine/Plugins/Media/MediaPlate/Source/MediaPlate/Public/MediaPlateComponent.h b/Engine/Plugins/Media/MediaPlate/Source/MediaPlate/Public/MediaPlateComponent.h
--- a/Engine/Plugins/Media/MediaPlate/Source/MediaPlate/Public/MediaPlateComponent.h
+++ b/Engine/Plugins/Media/MediaPlate/Source/MediaPlate/Public/MediaPlateComponent.h
@@ -79,16 +79,8 @@ public:
 	UPROPERTY(EditAnywhere, AdvancedDisplay, Category = MediaPlate)
 	TObjectPtr<UMediaSoundComponent> SoundComponent;
 
-	/** Selects whether to use the media source or the media path. */
-	UPROPERTY(EditAnywhere, Category = MediaPlate)
-	bool bUseMediaSource;
-
-	/** URL (or file) to play. */
-	UPROPERTY(EditAnywhere, Category = MediaPlate, meta = (EditCondition = "!bUseMediaSource"))
-	FFilePath MediaPath;
-
 	/** What media to play. */
-	UPROPERTY(EditAnywhere, Category = MediaPlate, meta = (EditCondition = "bUseMediaSource"))
+	UPROPERTY(EditAnywhere, Category = "MediaPlate")
 	TObjectPtr<UMediaSource> MediaSource;
 
 	/** Enable smart caching for image sequences. */
@@ -119,10 +111,6 @@ private:
 	/** Info representing this object. */
 	TSharedPtr<FMediaTextureTrackerObject, ESPMode::ThreadSafe> MediaTextureTrackerObject;
 
-	/** If we are using MediaPath, then this is the media source for it. */
-	UPROPERTY(Transient)
-	TObjectPtr<UMediaSource> MediaPathMediaSource;
-
 	/**
 	 * Plays a media source.
 	 * 
diff --git a/Engine/Plugins/Media/MediaPlate/Source/MediaPlateEditor/MediaPlateEditor.Build.cs b/Engine/Plugins/Media/MediaPlate/Source/MediaPlateEditor/MediaPlateEditor.Build.cs
--- a/Engine/Plugins/Media/MediaPlate/Source/MediaPlateEditor/MediaPlateEditor.Build.cs
+++ b/Engine/Plugins/Media/MediaPlate/Source/MediaPlateEditor/MediaPlateEditor.Build.cs
@@ -19,6 +19,7 @@ public class MediaPlateEditor : ModuleRules
 			new string[]
 			{
 				"CoreUObject",
+				"DesktopWidgets",
 				"EditorStyle",
 				"Engine",
 				"InputCore",
diff --git a/Engine/Plugins/Media/MediaPlate/Source/MediaPlateEditor/Private/MediaPlateCustomization.cpp b/Engine/Plugins/Media/MediaPlate/Source/MediaPlateEditor/Private/MediaPlateCustomization.cpp
--- a/Engine/Plugins/Media/MediaPlate/Source/MediaPlateEditor/Private/MediaPlateCustomization.cpp
+++ b/Engine/Plugins/Media/MediaPlate/Source/MediaPlateEditor/Private/MediaPlateCustomization.cpp
@@ -10,8 +10,10 @@
 #include "MediaPlateComponent.h"
 #include "MediaPlateEditorModule.h"
 #include "MediaPlayer.h"
+#include "MediaSource.h"
 #include "Subsystems/AssetEditorSubsystem.h"
 #include "Widgets/Input/SButton.h"
+#include "Widgets/Input/SFilePathPicker.h"
 #include "Widgets/Images/SImage.h"
 
 #define LOCTEXT_NAMESPACE "FMediaPlateCustomization"
@@ -53,17 +55,42 @@ void FMediaPlateCustomization::CustomizeDetails(IDetailLayoutBuilder& DetailBuil
 		}
 	}
 
-	// Get the bUseMediaSource property.
-	TSharedRef<IPropertyHandle> Property = DetailBuilder.GetProperty(GET_MEMBER_NAME_CHECKED(UMediaPlateComponent, bUseMediaSource));
+	// Set media path.
+	UpdateMediaPath();
+
+	// Get the MediaSource property.
+	TSharedRef<IPropertyHandle> Property = DetailBuilder.GetProperty(GET_MEMBER_NAME_CHECKED(UMediaPlateComponent, MediaSource));
 	if (Property->IsValidHandle())
 	{
-		
 		// Get a callback when this changes so we can stop the player.
-		FSimpleDelegate OnUseMediaSourceChangedDelegate = FSimpleDelegate::CreateSP(this,
-			&FMediaPlateCustomization::OnUseMediaSourceChanged, &DetailBuilder);
-		Property->SetOnPropertyValueChanged(OnUseMediaSourceChangedDelegate);
+		FSimpleDelegate OnMediaSourceChangedDelegate = FSimpleDelegate::CreateSP(this,
+			&FMediaPlateCustomization::OnMediaSourceChanged, &DetailBuilder);
+		Property->SetOnPropertyValueChanged(OnMediaSourceChangedDelegate);
 	}
 
+	// Add media path.
+	FString FileTypeFilter = TEXT("All files (*.*)|*.*");
+	MediaPlateCategory.AddCustomRow(LOCTEXT("MediaPlateMediaPath", "MediaPlate Media Path"))
+		.NameContent()
+		[
+			SNew(STextBlock)
+				.Text(LOCTEXT("MediaPath", "Media Path"))
+				.ToolTipText(LOCTEXT("MediaPath_ToolTip",
+					"The path of the Media Source to play.\nChanging this will create a new media source in the level to play this path."))
+				.Font(IDetailLayoutBuilder::GetDetailFont())
+		]
+		.ValueContent()
+		[
+			SNew(SFilePathPicker)
+				.BrowseButtonImage(FEditorStyle::GetBrush("PropertyWindow.Button_Ellipsis"))
+				.BrowseButtonStyle(FEditorStyle::Get(), "HoverHintOnly")
+				.BrowseButtonToolTip(LOCTEXT("FileButtonToolTipText", "Choose a file from this computer"))
+				.BrowseTitle(LOCTEXT("PropertyEditorTitle", "File picker..."))
+				.FilePath(this, &FMediaPlateCustomization::HandleMediaPath)
+				.FileTypeFilter(FileTypeFilter)
+				.OnPathPicked(this, &FMediaPlateCustomization::HandleMediaPathPicked)
+		];
+
 	// Add media control buttons.
 	MediaPlateCategory.AddCustomRow(LOCTEXT("MediaPlateControls", "MediaPlate Controls"))
 		[
@@ -284,6 +311,71 @@ void FMediaPlateCustomization::CustomizeDetails(IDetailLayoutBuilder& DetailBuil
 		}
 	}
 
+
+void FMediaPlateCustomization::UpdateMediaPath()
+{
+	MediaPath.Empty();
+
+	if (MediaPlatesList.Num() > 0)
+	{
+		UMediaPlateComponent* MediaPlate = MediaPlatesList[0].Get();
+		if (MediaPlate != nullptr)
+		{
+			UMediaSource* MediaSource = MediaPlate->MediaSource;
+			if (MediaSource != nullptr)
+			{
+				MediaPath = MediaSource->GetUrl();
+
+				// Remove certain types.
+				const FString FilePrefix(TEXT("file://"));
+				const FString ImgPrefix(TEXT("img://"));
+				if (MediaPath.StartsWith(FilePrefix))
+				{
+					MediaPath = MediaPath.RightChop(FilePrefix.Len());
+				}
+				else if (MediaPath.StartsWith(ImgPrefix))
+				{
+					MediaPath = MediaPath.RightChop(ImgPrefix.Len());
+				}
+			}
+		}
+	}
+}
+
+FString FMediaPlateCustomization::HandleMediaPath() const
+{
+	return MediaPath;
+}
+
+void FMediaPlateCustomization::HandleMediaPathPicked(const FString& PickedPath)
+{
+	// Did we get something?
+	if ((PickedPath.IsEmpty() == false) && (PickedPath != MediaPath))
+	{
+		// Stop playback.
+		StopMediaPlates();
+
+		// Set up media source for our media plates.
+		for (TWeakObjectPtr<UMediaPlateComponent>& MediaPlatePtr : MediaPlatesList)
+		{
+			UMediaPlateComponent* MediaPlate = MediaPlatePtr.Get();
+			if (MediaPlate != nullptr)
+			{
+				// Create media source for this path.
+				UMediaSource* MediaSource = UMediaSource::SpawnMediaSourceForString(PickedPath, MediaPlate);
+				if (MediaSource != nullptr)
+				{
+					MediaPlate->MediaSource = MediaSource;
+					MediaPlate->MarkPackageDirty();
+				}
+			}
+		}
+
+		// Update the media path.
+		UpdateMediaPath();
+	}
+}
+
 FReply FMediaPlateCustomization::OnOpenMediaPlate()
 {
 	// Get all our objects.
@@ -306,9 +398,8 @@ FReply FMediaPlateCustomization::OnOpenMediaPlate()
 	return FReply::Handled();
 }
 
-void FMediaPlateCustomization::OnUseMediaSourceChanged(IDetailLayoutBuilder* DetailBuilder)
+void FMediaPlateCustomization::StopMediaPlates()
 {
-	// Stop the player if we change the media source.
 	for (TWeakObjectPtr<UMediaPlateComponent>& MediaPlatePtr : MediaPlatesList)
 	{
 		UMediaPlateComponent* MediaPlate = MediaPlatePtr.Get();
@@ -319,6 +410,12 @@ void FMediaPlateCustomization::OnUseMediaSourceChanged(IDetailLayoutBuilder* Det
 	}
 }
 
+void FMediaPlateCustomization::OnMediaSourceChanged(IDetailLayoutBuilder* DetailBuilder)
+{
+	StopMediaPlates();
+	UpdateMediaPath();
+}
+
 float FMediaPlateCustomization::GetForwardRate(UMediaPlayer* MediaPlayer) const
 {
 	float Rate = MediaPlayer->GetRate();
diff --git a/Engine/Plugins/Media/MediaPlate/Source/MediaPlateEditor/Private/MediaPlateCustomization.h b/Engine/Plugins/Media/MediaPlate/Source/MediaPlateEditor/Private/MediaPlateCustomization.h
--- a/Engine/Plugins/Media/MediaPlate/Source/MediaPlateEditor/Private/MediaPlateCustomization.h
+++ b/Engine/Plugins/Media/MediaPlate/Source/MediaPlateEditor/Private/MediaPlateCustomization.h
@@ -37,6 +37,23 @@ private:
 
 	/** List of the media plates we are editing. */
 	TArray<TWeakObjectPtr<UMediaPlateComponent>> MediaPlatesList;
+	/** Stores the current value of the MediaPath property. */
+	FString MediaPath;
+
+	/**
+	 * Updates MediaPath from the current MediaSource.
+	 */
+	void UpdateMediaPath();
+
+	/**
+	 * Called to get the media path for the file picker.
+	 */
+	FString HandleMediaPath() const;
+	
+	/**
+	 * Called when we select a media path.
+	 */
+	void HandleMediaPathPicked(const FString& PickedPath);
 
 	/**
 	 * Called when the open media plate button is pressed.
@@ -44,9 +61,14 @@ private:
 	FReply OnOpenMediaPlate();
 
 	/**
-	 * Called when bUseMediaSource changes.
+	 * Call this to stop all playback.
+	 */
+	void StopMediaPlates();
+
+	/**
+	 * Called when MediaSource changes.
 	 */
-	void OnUseMediaSourceChanged(IDetailLayoutBuilder* DetailBuilder);
+	void OnMediaSourceChanged(IDetailLayoutBuilder* DetailBuilder);
 
 	/**
 	 * Get the rate to use when we press the forward button.
diff --git a/Engine/Source/Runtime/MediaAssets/Private/Assets/MediaSource.cpp b/Engine/Source/Runtime/MediaAssets/Private/Assets/MediaSource.cpp
--- a/Engine/Source/Runtime/MediaAssets/Private/Assets/MediaSource.cpp
+++ b/Engine/Source/Runtime/MediaAssets/Private/Assets/MediaSource.cpp
@@ -24,7 +24,7 @@ void UMediaSource::UnregisterSpawnFromFileExtension(const FString& Extension)
 	Delegates.Remove(Extension);
 }
 
-UMediaSource* UMediaSource::SpawnMediaSourceForString(const FString& MediaPath)
+UMediaSource* UMediaSource::SpawnMediaSourceForString(const FString& MediaPath, UObject* Outer)
 {
 	TObjectPtr<UMediaSource> MediaSource = nullptr;
 
@@ -32,7 +32,7 @@ UMediaSource* UMediaSource::SpawnMediaSourceForString(const FString& MediaPath)
 	bool bIsUrl = MediaPath.Contains(TEXT("://"));
 	if (bIsUrl)
 	{
-		TObjectPtr<UStreamMediaSource> StreamMediaSource = NewObject<UStreamMediaSource>(GetTransientPackage(), NAME_None, RF_Transactional | RF_Transient);
+		TObjectPtr<UStreamMediaSource> StreamMediaSource = NewObject<UStreamMediaSource>(Outer, NAME_None, RF_Transactional);
 		StreamMediaSource->StreamUrl = MediaPath;
 		MediaSource = StreamMediaSource;
 	}
@@ -45,12 +45,12 @@ UMediaSource* UMediaSource::SpawnMediaSourceForString(const FString& MediaPath)
 		FMediaSourceSpawnDelegate* Delegate = Delegates.Find(FileExtension);
 		if ((Delegate != nullptr) && (Delegate->IsBound()))
 		{
-			MediaSource = Delegate->Execute(MediaPath);
+			MediaSource = Delegate->Execute(MediaPath, Outer);
 		}
 		else
 		{
 			// Try a file media source.
-			TObjectPtr<UFileMediaSource> FileMediaSource = NewObject<UFileMediaSource>(GetTransientPackage(), NAME_None, RF_Transactional | RF_Transient);
+			TObjectPtr<UFileMediaSource> FileMediaSource = NewObject<UFileMediaSource>(Outer, NAME_None, RF_Transactional);
 			FileMediaSource->SetFilePath(MediaPath);
 			MediaSource = FileMediaSource;
 		}
diff --git a/Engine/Source/Runtime/MediaAssets/Public/MediaSource.h b/Engine/Source/Runtime/MediaAssets/Public/MediaSource.h
--- a/Engine/Source/Runtime/MediaAssets/Public/MediaSource.h
+++ b/Engine/Source/Runtime/MediaAssets/Public/MediaSource.h
@@ -13,7 +13,7 @@
 #include "MediaSource.generated.h"
 
 /** Delegate for creating a media source from a string. */
-DECLARE_DELEGATE_RetVal_OneParam(UMediaSource*, FMediaSourceSpawnDelegate, const FString&);
+DECLARE_DELEGATE_RetVal_TwoParams(UMediaSource*, FMediaSourceSpawnDelegate, const FString&, UObject*);
 
 /**
  * Abstract base class for media sources.
@@ -71,9 +71,10 @@ public:
 	 * Call this to try and create a media source appropriate for the media.
 	 *
 	 * @param MediaPath		Can be a file location or a Url.
+	 * @param Outer			Outer to use for this object.
 	 * @return				Media source or nullptr if none are appropriate.
 	 */
-	static UMediaSource* SpawnMediaSourceForString(const FString& MediaPath);
+	static UMediaSource* SpawnMediaSourceForString(const FString& MediaPath, UObject* Outer);
 
 public:
 
