MediaPlate: Added a playlist to the media plate.

#rnx

diff --git a/Engine/Plugins/Media/MediaPlate/Source/MediaPlate/Private/MediaPlateComponent.cpp b/Engine/Plugins/Media/MediaPlate/Source/MediaPlate/Private/MediaPlateComponent.cpp
--- a/Engine/Plugins/Media/MediaPlate/Source/MediaPlate/Private/MediaPlateComponent.cpp
+++ b/Engine/Plugins/Media/MediaPlate/Source/MediaPlate/Private/MediaPlateComponent.cpp
@@ -5,8 +5,9 @@
 #include "Components/StaticMeshComponent.h"
 #include "Engine/StaticMesh.h"
 #include "MediaComponent.h"
-#include "MediaPlayer.h"
 #include "MediaPlateModule.h"
+#include "MediaPlayer.h"
+#include "MediaPlaylist.h"
 #include "MediaSource.h"
 #include "MediaSoundComponent.h"
 #include "MediaTexture.h"
@@ -15,6 +16,7 @@
 #define LOCTEXT_NAMESPACE "MediaPlate"
 
 FLazyName UMediaPlateComponent::MediaComponentName(TEXT("MediaComponent0"));
+FLazyName UMediaPlateComponent::MediaPlaylistName(TEXT("MediaPlaylist0"));
 
 UMediaPlateComponent::UMediaPlateComponent(const FObjectInitializer& ObjectInitializer)
 	: Super(ObjectInitializer)
@@ -30,6 +32,9 @@ UMediaPlateComponent::UMediaPlateComponent(const FObjectInitializer& ObjectIniti
 			MediaTexture->NewStyleOutput = true;
 		}
 	}
+
+	// Set up playlist.
+	MediaPlaylist = CreateDefaultSubobject<UMediaPlaylist>(MediaPlaylistName);
 }
 
 void UMediaPlateComponent::OnRegister()
@@ -86,6 +91,11 @@ void UMediaPlateComponent::Play()
 	TObjectPtr<UMediaPlayer> MediaPlayer = GetMediaPlayer();
 	if (MediaPlayer != nullptr)
 	{
+		UMediaSource* MediaSource = nullptr;
+		if (MediaPlaylist != nullptr)
+		{
+			MediaSource = MediaPlaylist->Get(0);
+		}
 		bIsPlaying = PlayMediaSource(MediaSource);
 	}
 
diff --git a/Engine/Plugins/Media/MediaPlate/Source/MediaPlate/Public/MediaPlateComponent.h b/Engine/Plugins/Media/MediaPlate/Source/MediaPlate/Public/MediaPlateComponent.h
--- a/Engine/Plugins/Media/MediaPlate/Source/MediaPlate/Public/MediaPlateComponent.h
+++ b/Engine/Plugins/Media/MediaPlate/Source/MediaPlate/Public/MediaPlateComponent.h
@@ -10,6 +10,7 @@
 
 class UMediaComponent;
 class UMediaPlayer;
+class UMediaPlaylist;
 class UMediaSoundComponent;
 class UMediaSource;
 class UMediaTexture;
@@ -79,9 +80,9 @@ public:
 	UPROPERTY(EditAnywhere, AdvancedDisplay, Category = MediaPlate)
 	TObjectPtr<UMediaSoundComponent> SoundComponent;
 
-	/** What media to play. */
+	/** What media playlist to play. */
 	UPROPERTY(EditAnywhere, Category = "MediaPlate")
-	TObjectPtr<UMediaSource> MediaSource;
+	TObjectPtr<UMediaPlaylist> MediaPlaylist;
 
 	/** Enable smart caching for image sequences. */
 	UPROPERTY(EditAnywhere, AdvancedDisplay, Category = "MediaPlate")
@@ -107,6 +108,8 @@ public:
 private:
 	/** Name for our media component. */
 	static FLazyName MediaComponentName;
+	/** Name for our playlist. */
+	static FLazyName MediaPlaylistName;
 
 	/** Info representing this object. */
 	TSharedPtr<FMediaTextureTrackerObject, ESPMode::ThreadSafe> MediaTextureTrackerObject;
diff --git a/Engine/Plugins/Media/MediaPlate/Source/MediaPlateEditor/MediaPlateEditor.Build.cs b/Engine/Plugins/Media/MediaPlate/Source/MediaPlateEditor/MediaPlateEditor.Build.cs
--- a/Engine/Plugins/Media/MediaPlate/Source/MediaPlateEditor/MediaPlateEditor.Build.cs
+++ b/Engine/Plugins/Media/MediaPlate/Source/MediaPlateEditor/MediaPlateEditor.Build.cs
@@ -29,6 +29,7 @@ public class MediaPlateEditor : ModuleRules
 				"MediaPlate",
 				"MediaPlayerEditor",
 				"MovieScene",
+				"PropertyEditor",
 				"Sequencer",
 				"Slate",
 				"SlateCore",
diff --git a/Engine/Plugins/Media/MediaPlate/Source/MediaPlateEditor/Private/ActorFactoryMediaPlate.cpp b/Engine/Plugins/Media/MediaPlate/Source/MediaPlateEditor/Private/ActorFactoryMediaPlate.cpp
--- a/Engine/Plugins/Media/MediaPlate/Source/MediaPlateEditor/Private/ActorFactoryMediaPlate.cpp
+++ b/Engine/Plugins/Media/MediaPlate/Source/MediaPlateEditor/Private/ActorFactoryMediaPlate.cpp
@@ -5,6 +5,7 @@
 #include "MediaSource.h"
 #include "MediaPlate.h"
 #include "MediaPlateComponent.h"
+#include "MediaPlaylist.h"
 
 #define LOCTEXT_NAMESPACE "ActorFactoryMediaPlate"
 
@@ -57,7 +58,11 @@ void UActorFactoryMediaPlate::SetUpActor(UObject* Asset, AActor* Actor)
 		UMediaSource* MediaSource = Cast<UMediaSource>(Asset);
 		if ((MediaSource != nullptr) && (MediaPlate->MediaPlateComponent != nullptr))
 		{
-			MediaPlate->MediaPlateComponent->MediaSource = MediaSource;
+			UMediaPlaylist* Playlist = MediaPlate->MediaPlateComponent->MediaPlaylist;
+			if (Playlist != nullptr)
+			{
+				Playlist->Add(MediaSource);
+			}
 		}
 	}
 }
diff --git a/Engine/Plugins/Media/MediaPlate/Source/MediaPlateEditor/Private/MediaPlateCustomization.cpp b/Engine/Plugins/Media/MediaPlate/Source/MediaPlateEditor/Private/MediaPlateCustomization.cpp
--- a/Engine/Plugins/Media/MediaPlate/Source/MediaPlateEditor/Private/MediaPlateCustomization.cpp
+++ b/Engine/Plugins/Media/MediaPlate/Source/MediaPlateEditor/Private/MediaPlateCustomization.cpp
@@ -10,7 +10,9 @@
 #include "MediaPlateComponent.h"
 #include "MediaPlateEditorModule.h"
 #include "MediaPlayer.h"
+#include "MediaPlaylist.h"
 #include "MediaSource.h"
+#include "PropertyCustomizationHelpers.h"
 #include "Subsystems/AssetEditorSubsystem.h"
 #include "Widgets/Input/SButton.h"
 #include "Widgets/Input/SFilePathPicker.h"
@@ -57,16 +59,23 @@ void FMediaPlateCustomization::CustomizeDetails(IDetailLayoutBuilder& DetailBuil
 
 	// Set media path.
 	UpdateMediaPath();
-
-	// Get the MediaSource property.
-	TSharedRef<IPropertyHandle> Property = DetailBuilder.GetProperty(GET_MEMBER_NAME_CHECKED(UMediaPlateComponent, MediaSource));
-	if (Property->IsValidHandle())
-	{
-		// Get a callback when this changes so we can stop the player.
-		FSimpleDelegate OnMediaSourceChangedDelegate = FSimpleDelegate::CreateSP(this,
-			&FMediaPlateCustomization::OnMediaSourceChanged, &DetailBuilder);
-		Property->SetOnPropertyValueChanged(OnMediaSourceChangedDelegate);
-	}
+	
+	// Add media source.
+	MediaPlateCategory.AddCustomRow(LOCTEXT("MediaPlateMediaSource", "MediaPlate Media Source"))
+		.NameContent()
+		[
+			SNew(STextBlock)
+			.Text(LOCTEXT("MediaSource", "Media Source"))
+			.ToolTipText(LOCTEXT("MediaSource_ToolTip", "The Media Source to play."))
+			.Font(IDetailLayoutBuilder::GetDetailFont())
+		]
+		.ValueContent()
+		[
+			SNew(SObjectPropertyEntryBox)
+				.AllowedClass(UMediaSource::StaticClass())
+				.ObjectPath(this, &FMediaPlateCustomization::GetMediaSourcePath)
+				.OnObjectChanged(this, &FMediaPlateCustomization::OnMediaSourceChanged)
+		];
 
 	// Add media path.
 	FString FileTypeFilter = TEXT("All files (*.*)|*.*");
@@ -308,34 +317,94 @@ void FMediaPlateCustomization::CustomizeDetails(IDetailLayoutBuilder& DetailBuil
 							.Text(LOCTEXT("OpenMediaPlate", "Open Media Plate"))
 					]
 			];
+	}
+}
+
+FString FMediaPlateCustomization::GetMediaSourcePath() const
+{
+	FString Path;
+
+	// Get the first media plate.
+	if (MediaPlatesList.Num() > 0)
+	{
+		UMediaPlateComponent* MediaPlate = MediaPlatesList[0].Get();
+		if (MediaPlate != nullptr)
+		{
+			UMediaPlaylist* Playlist = MediaPlate->MediaPlaylist;
+			if (Playlist != nullptr)
+			{
+				// Get the first media source in the playlist.
+				UMediaSource* MediaSource = Playlist->Get(0);
+				if (MediaSource != nullptr)
+				{
+					Path = MediaSource->GetPathName();
+				}
+			}
 		}
 	}
 
+	return Path;
+}
+
+void FMediaPlateCustomization::OnMediaSourceChanged(const FAssetData& AssetData)
+{
+	// Update the playlist with the new media source.
+	UMediaSource* MediaSource = Cast<UMediaSource>(AssetData.GetAsset());
+	for (TWeakObjectPtr<UMediaPlateComponent>& MediaPlatePtr : MediaPlatesList)
+	{
+		UMediaPlateComponent* MediaPlate = MediaPlatePtr.Get();
+		if (MediaPlate != nullptr)
+		{
+			UMediaPlaylist* Playlist = MediaPlate->MediaPlaylist;
+			if (Playlist != nullptr)
+			{
+				if (Playlist->Num() > 0)
+				{
+					Playlist->Replace(0, MediaSource);
+				}
+				else
+				{
+					Playlist->Add(MediaSource);
+				}
+				Playlist->MarkPackageDirty();
+			}
+		}
+	}
+
+	StopMediaPlates();
+	UpdateMediaPath();
+}
 
 void FMediaPlateCustomization::UpdateMediaPath()
 {
 	MediaPath.Empty();
 
+	// Get the first media plate.
 	if (MediaPlatesList.Num() > 0)
 	{
 		UMediaPlateComponent* MediaPlate = MediaPlatesList[0].Get();
 		if (MediaPlate != nullptr)
 		{
-			UMediaSource* MediaSource = MediaPlate->MediaSource;
-			if (MediaSource != nullptr)
+			UMediaPlaylist* Playlist = MediaPlate->MediaPlaylist;
+			if (Playlist != nullptr)
 			{
-				MediaPath = MediaSource->GetUrl();
-
-				// Remove certain types.
-				const FString FilePrefix(TEXT("file://"));
-				const FString ImgPrefix(TEXT("img://"));
-				if (MediaPath.StartsWith(FilePrefix))
-				{
-					MediaPath = MediaPath.RightChop(FilePrefix.Len());
-				}
-				else if (MediaPath.StartsWith(ImgPrefix))
+				// Get the first media source in the playlist.
+				UMediaSource* MediaSource = Playlist->Get(0);
+				if (MediaSource != nullptr)
 				{
-					MediaPath = MediaPath.RightChop(ImgPrefix.Len());
+					MediaPath = MediaSource->GetUrl();
+
+					// Remove certain types.
+					const FString FilePrefix(TEXT("file://"));
+					const FString ImgPrefix(TEXT("img://"));
+					if (MediaPath.StartsWith(FilePrefix))
+					{
+						MediaPath = MediaPath.RightChop(FilePrefix.Len());
+					}
+					else if (MediaPath.StartsWith(ImgPrefix))
+					{
+						MediaPath = MediaPath.RightChop(ImgPrefix.Len());
+					}
 				}
 			}
 		}
@@ -361,12 +430,23 @@ void FMediaPlateCustomization::HandleMediaPathPicked(const FString& PickedPath)
 			UMediaPlateComponent* MediaPlate = MediaPlatePtr.Get();
 			if (MediaPlate != nullptr)
 			{
-				// Create media source for this path.
-				UMediaSource* MediaSource = UMediaSource::SpawnMediaSourceForString(PickedPath, MediaPlate);
-				if (MediaSource != nullptr)
+				UMediaPlaylist* Playlist = MediaPlate->MediaPlaylist;
+				if (Playlist != nullptr)
 				{
-					MediaPlate->MediaSource = MediaSource;
-					MediaPlate->MarkPackageDirty();
+					// Create media source for this path.
+					UMediaSource* MediaSource = UMediaSource::SpawnMediaSourceForString(PickedPath, MediaPlate);
+					if (MediaSource != nullptr)
+					{
+						if (Playlist->Num() > 0)
+						{
+							Playlist->Replace(0, MediaSource);
+						}
+						else
+						{
+							Playlist->Add(MediaSource);
+						}
+						Playlist->MarkPackageDirty();
+					}
 				}
 			}
 		}
@@ -410,12 +490,6 @@ void FMediaPlateCustomization::StopMediaPlates()
 	}
 }
 
-void FMediaPlateCustomization::OnMediaSourceChanged(IDetailLayoutBuilder* DetailBuilder)
-{
-	StopMediaPlates();
-	UpdateMediaPath();
-}
-
 float FMediaPlateCustomization::GetForwardRate(UMediaPlayer* MediaPlayer) const
 {
 	float Rate = MediaPlayer->GetRate();
diff --git a/Engine/Plugins/Media/MediaPlate/Source/MediaPlateEditor/Private/MediaPlateCustomization.h b/Engine/Plugins/Media/MediaPlate/Source/MediaPlateEditor/Private/MediaPlateCustomization.h
--- a/Engine/Plugins/Media/MediaPlate/Source/MediaPlateEditor/Private/MediaPlateCustomization.h
+++ b/Engine/Plugins/Media/MediaPlate/Source/MediaPlateEditor/Private/MediaPlateCustomization.h
@@ -40,6 +40,16 @@ private:
 	/** Stores the current value of the MediaPath property. */
 	FString MediaPath;
 
+	/**
+	 * Gets the object path for the media source object.
+	 */
+	FString GetMediaSourcePath() const;
+
+	/**
+	 * Called when the media source widget changes.
+	 */
+	void OnMediaSourceChanged(const FAssetData& AssetData);
+
 	/**
 	 * Updates MediaPath from the current MediaSource.
 	 */
@@ -65,11 +75,6 @@ private:
 	 */
 	void StopMediaPlates();
 
-	/**
-	 * Called when MediaSource changes.
-	 */
-	void OnMediaSourceChanged(IDetailLayoutBuilder* DetailBuilder);
-
 	/**
 	 * Get the rate to use when we press the forward button.
 	 */
