MediaPlate: Moved the playlist and media source/path to a playlist category to make it more obvious what is going on.
#rnx

diff --git a/Engine/Plugins/Media/MediaPlate/Source/MediaPlate/Public/MediaPlateComponent.h b/Engine/Plugins/Media/MediaPlate/Source/MediaPlate/Public/MediaPlateComponent.h
--- a/Engine/Plugins/Media/MediaPlate/Source/MediaPlate/Public/MediaPlateComponent.h
+++ b/Engine/Plugins/Media/MediaPlate/Source/MediaPlate/Public/MediaPlateComponent.h
@@ -81,7 +81,7 @@ public:
 	TObjectPtr<UMediaSoundComponent> SoundComponent;
 
 	/** What media playlist to play. */
-	UPROPERTY(EditAnywhere, Category = "MediaPlate")
+	UPROPERTY(EditAnywhere, Category = "MediaPlate Playlist")
 	TObjectPtr<UMediaPlaylist> MediaPlaylist;
 
 	/** Enable smart caching for image sequences. */
diff --git a/Engine/Plugins/Media/MediaPlate/Source/MediaPlateEditor/Private/MediaPlateCustomization.cpp b/Engine/Plugins/Media/MediaPlate/Source/MediaPlateEditor/Private/MediaPlateCustomization.cpp
--- a/Engine/Plugins/Media/MediaPlate/Source/MediaPlateEditor/Private/MediaPlateCustomization.cpp
+++ b/Engine/Plugins/Media/MediaPlate/Source/MediaPlateEditor/Private/MediaPlateCustomization.cpp
@@ -60,8 +60,10 @@ void FMediaPlateCustomization::CustomizeDetails(IDetailLayoutBuilder& DetailBuil
 	// Set media path.
 	UpdateMediaPath();
 	
+	IDetailCategoryBuilder& MediaPlaylistCategory = DetailBuilder.EditCategory("MediaPlate Playlist");
+
 	// Add media source.
-	MediaPlateCategory.AddCustomRow(LOCTEXT("MediaPlateMediaSource", "MediaPlate Media Source"))
+	MediaPlaylistCategory.AddCustomRow(LOCTEXT("MediaPlateMediaSource", "MediaPlate Media Source"))
 		.NameContent()
 		[
 			SNew(STextBlock)
@@ -79,7 +81,7 @@ void FMediaPlateCustomization::CustomizeDetails(IDetailLayoutBuilder& DetailBuil
 
 	// Add media path.
 	FString FileTypeFilter = TEXT("All files (*.*)|*.*");
-	MediaPlateCategory.AddCustomRow(LOCTEXT("MediaPlateMediaPath", "MediaPlate Media Path"))
+	MediaPlaylistCategory.AddCustomRow(LOCTEXT("MediaPlateMediaPath", "MediaPlate Media Path"))
 		.NameContent()
 		[
 			SNew(STextBlock)
