MediaPlate: Media track is automatically populated when adding a media plate to Sequencer.
#rnx

diff --git a/Engine/Plugins/Media/MediaPlate/Source/MediaPlateEditor/Private/Sequencer/MediaPlateTrackEditor.cpp b/Engine/Plugins/Media/MediaPlate/Source/MediaPlateEditor/Private/Sequencer/MediaPlateTrackEditor.cpp
--- a/Engine/Plugins/Media/MediaPlate/Source/MediaPlateEditor/Private/Sequencer/MediaPlateTrackEditor.cpp
+++ b/Engine/Plugins/Media/MediaPlate/Source/MediaPlateEditor/Private/Sequencer/MediaPlateTrackEditor.cpp
@@ -7,6 +7,7 @@
 #include "MediaTexture.h"
 #include "MediaPlate.h"
 #include "MediaPlateComponent.h"
+#include "MediaPlaylist.h"
 #include "MovieSceneMediaTrack.h"
 
 #define LOCTEXT_NAMESPACE "FMediaPlateTrackEditor"
@@ -113,6 +114,20 @@ void FMediaPlateTrackEditor::AddTrackForComponent(UMediaPlateComponent* Componen
 	{
 		UMovieSceneMediaTrack* MediaTrack = Cast<UMovieSceneMediaTrack>(Track);
 		MediaTrack->SetDisplayName(LOCTEXT("MediaTrackName", "Media"));
+
+		// Populate track.
+		UMediaPlaylist* Playlist = Component->MediaPlaylist;
+		if (Playlist != nullptr)
+		{
+			for (int32 Index = 0; Index < Playlist->Num(); ++Index)
+			{
+				UMediaSource* MediaSource = Playlist->Get(Index);
+				if (MediaSource != nullptr)
+				{
+					MediaTrack->AddNewMediaSource(*MediaSource, FFrameNumber(0));
+				}
+			}
+		}
 	}
 }
 
