MediaPlate: Added playlist tracks view to toolkit.

#rnx

diff --git a/Engine/Plugins/Media/MediaPlate/Source/MediaPlateEditor/Private/Toolkits/MediaPlateEditorToolkit.cpp b/Engine/Plugins/Media/MediaPlate/Source/MediaPlateEditor/Private/Toolkits/MediaPlateEditorToolkit.cpp
--- a/Engine/Plugins/Media/MediaPlate/Source/MediaPlateEditor/Private/Toolkits/MediaPlateEditorToolkit.cpp
+++ b/Engine/Plugins/Media/MediaPlate/Source/MediaPlateEditor/Private/Toolkits/MediaPlateEditorToolkit.cpp
@@ -17,6 +17,7 @@
 #include "Widgets/SMediaPlateEditorDetails.h"
 #include "Widgets/SMediaPlateEditorMediaDetails.h"
 #include "Widgets/SMediaPlayerEditorViewer.h"
+#include "Widgets/SMediaPlaylistEditorTracks.h"
 
 #define LOCTEXT_NAMESPACE "FMediaPlateEditorToolkit"
 
@@ -25,6 +26,7 @@ namespace MediaPlateEditorToolkit
 	static const FName AppIdentifier("MediaPlateEditorApp");
 	static const FName DetailsTabId("Details");
 	static const FName MediaDetailsTabId("MediaDetails");
+	static const FName PlaylistTabId("Playlist");
 	static const FName ViewerTabId("Viewer");
 }
 
@@ -64,15 +66,16 @@ void FMediaPlateEditorToolkit::Initialize(UMediaPlateComponent* InMediaPlate, co
 	BindCommands();
 
 	// create tab layout
-	const TSharedRef<FTabManager::FLayout> Layout = FTabManager::NewLayout("Standalone_MediaPlateEditor_v1.2")
+	const TSharedRef<FTabManager::FLayout> Layout = FTabManager::NewLayout("Standalone_MediaPlateEditor_v1.3")
 		->AddArea
 		(
 			FTabManager::NewPrimaryArea()
-				->SetOrientation(Orient_Horizontal)
+				->SetOrientation(Orient_Vertical)
 				->Split
 				(
 					FTabManager::NewSplitter()
 						->SetOrientation(Orient_Horizontal)
+						->SetSizeCoefficient(0.7f)
 						->Split
 						(
 							// viewer
@@ -102,6 +105,13 @@ void FMediaPlateEditorToolkit::Initialize(UMediaPlateComponent* InMediaPlate, co
 							)
 						)
 				)
+				->Split
+				(
+					// Details tab.
+					FTabManager::NewStack()
+						->AddTab(MediaPlateEditorToolkit::PlaylistTabId, ETabState::OpenedTab)
+						->SetSizeCoefficient(0.3f)
+				)
 		);
 
 	FAssetEditorToolkit::InitAssetEditor(
@@ -161,6 +171,12 @@ void FMediaPlateEditorToolkit::RegisterTabSpawners(const TSharedRef<FTabManager>
 		.SetGroup(WorkspaceMenuCategoryRef)
 		.SetIcon(FSlateIcon(FEditorStyle::GetStyleSetName(), "LevelEditor.Tabs.Details"));
 
+	// Playlist tab.
+	InTabManager->RegisterTabSpawner(MediaPlateEditorToolkit::PlaylistTabId, FOnSpawnTab::CreateSP(this, &FMediaPlateEditorToolkit::HandleTabManagerSpawnTab, MediaPlateEditorToolkit::PlaylistTabId))
+		.SetDisplayName(LOCTEXT("PlaylistTabName", "Playlist"))
+		.SetGroup(WorkspaceMenuCategoryRef)
+		.SetIcon(FSlateIcon(FEditorStyle::GetStyleSetName(), "LevelEditor.Tabs.Details"));
+
 	// Viewer tab.
 	InTabManager->RegisterTabSpawner(MediaPlateEditorToolkit::ViewerTabId, FOnSpawnTab::CreateSP(this, &FMediaPlateEditorToolkit::HandleTabManagerSpawnTab, MediaPlateEditorToolkit::ViewerTabId))
 		.SetDisplayName(LOCTEXT("PlayerTabName", "Player"))
@@ -173,6 +189,7 @@ void FMediaPlateEditorToolkit::UnregisterTabSpawners(const TSharedRef<class FTab
 	FAssetEditorToolkit::UnregisterTabSpawners(InTabManager);
 
 	InTabManager->UnregisterTabSpawner(MediaPlateEditorToolkit::ViewerTabId);
+	InTabManager->UnregisterTabSpawner(MediaPlateEditorToolkit::PlaylistTabId);
 	InTabManager->UnregisterTabSpawner(MediaPlateEditorToolkit::MediaDetailsTabId);
 	InTabManager->UnregisterTabSpawner(MediaPlateEditorToolkit::DetailsTabId);
 }
@@ -382,6 +399,10 @@ TSharedRef<SDockTab> FMediaPlateEditorToolkit::HandleTabManagerSpawnTab(const FS
 	{
 		TabWidget = SNew(SMediaPlateEditorMediaDetails, *MediaPlate, Style);
 	}
+	else if (TabIdentifier == MediaPlateEditorToolkit::PlaylistTabId)
+	{
+		TabWidget = SNew(SMediaPlaylistEditorTracks, MediaPlate->MediaPlaylist, Style);
+	}
 	else if (TabIdentifier == MediaPlateEditorToolkit::ViewerTabId)
 	{
 		TabWidget = SNew(SMediaPlayerEditorViewer, *MediaPlayer, MediaTexture, Style, false);
diff --git a/Engine/Plugins/Media/MediaPlayerEditor/Source/MediaPlayerEditor/Private/Widgets/SMediaPlaylistEditorTracks.cpp b/Engine/Plugins/Media/MediaPlayerEditor/Source/MediaPlayerEditor/Private/Widgets/SMediaPlaylistEditorTracks.cpp
new file mode 100644
--- /dev/null
+++ b/Engine/Plugins/Media/MediaPlayerEditor/Source/MediaPlayerEditor/Private/Widgets/SMediaPlaylistEditorTracks.cpp
@@ -0,0 +1,135 @@
+// Copyright Epic Games, Inc. All Rights Reserved.
+
+#include "Widgets/SMediaPlaylistEditorTracks.h"
+#include "Modules/ModuleManager.h"
+#include "PropertyCustomizationHelpers.h"
+#include "PropertyEditorModule.h"
+#include "IDetailsView.h"
+#include "MediaPlaylist.h"
+#include "MediaSource.h"
+#include "Widgets/Input/SButton.h"
+
+#define LOCTEXT_NAMESPACE "SMediaPlaylistEditorTracks"
+
+void SMediaPlaylistEditorTracks::Construct(const FArguments& InArgs, UMediaPlaylist* InMediaPlaylist, const TSharedRef<ISlateStyle>& InStyle)
+{
+	MediaPlaylist = InMediaPlaylist;
+
+	ChildSlot
+		[
+			SNew(SVerticalBox)
+
+			// Buttons to manipulate playlist.
+			+ SVerticalBox::Slot()
+				.AutoHeight()
+				.Padding(2)
+				.VAlign(VAlign_Center)
+				[
+					SNew(SHorizontalBox)
+
+					+ SHorizontalBox::Slot()
+						.AutoWidth()
+						.Padding(2)
+						.HAlign(HAlign_Center)
+						[
+							// Add button to add to the playlist.
+							SNew(SButton)
+							.ToolTipText(LOCTEXT("MediaSource_ToolTip", "The Media Source to play."))
+								.VAlign(VAlign_Center)
+								.OnClicked_Lambda([this]() -> FReply
+								{
+									AddToPlaylist();
+									return FReply::Handled();
+								})
+								[
+									SNew(SImage)
+										.ColorAndOpacity(FSlateColor::UseForeground())
+										.Image(FEditorStyle::GetBrush("Icons.PlusCircle"))
+								]
+						]
+				]
+
+			// Container for our media sources.
+			+ SVerticalBox::Slot()
+				[
+					SAssignNew(SourcesContainer, SVerticalBox)
+				]
+		];
+
+	RefreshPlaylist();
+}
+
+void SMediaPlaylistEditorTracks::RefreshPlaylist()
+{
+	SourcesContainer->ClearChildren();
+
+	if (MediaPlaylist != nullptr)
+	{
+		// Add each item in the playlist.
+		for (int32 Index = 0; Index < MediaPlaylist->Num(); ++Index)
+		{
+			UMediaSource* MediaSource = MediaPlaylist->Get(Index);
+
+			SourcesContainer->AddSlot()
+				.AutoHeight()
+				.Padding(2)
+				.VAlign(VAlign_Center)
+				[
+					SNew(SHorizontalBox)
+					
+					+ SHorizontalBox::Slot()
+						.FillWidth(0.11f)
+						.Padding(2)
+						.HAlign(HAlign_Left)
+						[
+							// Add asset picker.
+							SNew(SObjectPropertyEntryBox)
+								.AllowedClass(UMediaSource::StaticClass())
+								.ObjectPath(this, &SMediaPlaylistEditorTracks::GetMediaSourcePath, Index)
+								.OnObjectChanged(this, &SMediaPlaylistEditorTracks::OnMediaSourceChanged, Index)
+						]
+				];
+		}
+	}
+}
+
+void SMediaPlaylistEditorTracks::AddToPlaylist()
+{
+	if (MediaPlaylist != nullptr)
+	{
+		MediaPlaylist->Insert(nullptr, MediaPlaylist->Num());
+		MediaPlaylist->MarkPackageDirty();
+		RefreshPlaylist();
+	}
+}
+
+FString SMediaPlaylistEditorTracks::GetMediaSourcePath(int32 Index) const
+{
+	FString Path;
+
+	if (MediaPlaylist != nullptr)
+	{
+		UMediaSource* MediaSource = MediaPlaylist->Get(Index);
+		if (MediaSource != nullptr)
+		{
+			Path = MediaSource->GetPathName();
+		}
+	}
+
+	return Path;
+}
+
+void SMediaPlaylistEditorTracks::OnMediaSourceChanged(const FAssetData& AssetData, int32 Index)
+{
+	if (MediaPlaylist != nullptr)
+	{
+		UMediaSource* MediaSource = Cast<UMediaSource>(AssetData.GetAsset());
+		
+		MediaPlaylist->Replace(Index, MediaSource);
+		MediaPlaylist->MarkPackageDirty();
+	}
+}
+
+
+
+#undef LOCTEXT_NAMESPACE
diff --git a/Engine/Plugins/Media/MediaPlayerEditor/Source/MediaPlayerEditor/Public/Widgets/SMediaPlaylistEditorTracks.h b/Engine/Plugins/Media/MediaPlayerEditor/Source/MediaPlayerEditor/Public/Widgets/SMediaPlaylistEditorTracks.h
new file mode 100644
--- /dev/null
+++ b/Engine/Plugins/Media/MediaPlayerEditor/Source/MediaPlayerEditor/Public/Widgets/SMediaPlaylistEditorTracks.h
@@ -0,0 +1,67 @@
+// Copyright Epic Games, Inc. All Rights Reserved.
+
+#pragma once
+
+#include "CoreMinimal.h"
+#include "Styling/ISlateStyle.h"
+#include "Widgets/DeclarativeSyntaxSupport.h"
+#include "Widgets/SCompoundWidget.h"
+
+class SVerticalBox;
+class UMediaPlaylist;
+
+/**
+ * Implements the tracks view of the MediaPlaylist.
+ */
+class MEDIAPLAYEREDITOR_API SMediaPlaylistEditorTracks
+	: public SCompoundWidget
+{
+public:
+
+	SLATE_BEGIN_ARGS(SMediaPlaylistEditorTracks) { }
+	SLATE_END_ARGS()
+
+public:
+
+	/**
+	 * Construct this widget
+	 *
+	 * @param InArgs The declaration data for this widget.
+	 * @param InMediaPlaylist The MediaPlaylist asset to show the details for.
+	 * @param InStyleSet The style set to use.
+	 */
+	void Construct(const FArguments& InArgs, UMediaPlaylist* InMediaPlaylist, const TSharedRef<ISlateStyle>& InStyle);
+
+private:
+
+	/** Pointer to the MediaPlaylist asset that is being viewed. */
+	TObjectPtr<UMediaPlaylist> MediaPlaylist;
+	/** Pointer to the container for the media sources. */
+	TSharedPtr<SVerticalBox> SourcesContainer;
+
+	/**
+	 * This should be called when the playlist changes so we can update the UI.
+	 */
+	void RefreshPlaylist();
+
+	/**
+	 * Adds a blank entry to the playlist.
+	 */
+	void AddToPlaylist();
+
+	/**
+	 * Gets the object path for the media source object.
+	 * 
+	 * @param	Index		Index in the playlist of the media source object.
+	 */
+	FString GetMediaSourcePath(int32 Index) const;
+
+	/**
+	 * Called when the media source widget changes.
+	 *
+	 * @param	AssetData	The media source.
+	 * @param	Index		Index in the playlist of the media source object.
+	 */
+	void OnMediaSourceChanged(const FAssetData& AssetData, int32 Index);
+
+};
