MediaPlate: Uses its own mesh so that the rotation for the previous floor mesh is unnecessary.
#rnx

diff --git a/Engine/Plugins/Media/MediaPlate/Source/MediaPlate/Private/MediaPlate.cpp b/Engine/Plugins/Media/MediaPlate/Source/MediaPlate/Private/MediaPlate.cpp
--- a/Engine/Plugins/Media/MediaPlate/Source/MediaPlate/Private/MediaPlate.cpp
+++ b/Engine/Plugins/Media/MediaPlate/Source/MediaPlate/Private/MediaPlate.cpp
@@ -26,14 +26,13 @@ AMediaPlate::AMediaPlate(const FObjectInitializer& ObjectInitializer)
 	// Set up static mesh component.
 	StaticMeshComponent = CreateDefaultSubobject<UStaticMeshComponent>(TEXT("StaticMeshComponent"));
 	StaticMeshComponent->SetupAttachment(RootComponent);
-	StaticMeshComponent->SetRelativeRotation(FRotator(0.0f, 90.0f, 90.0f));
 
 	// Hook up mesh.
 	struct FConstructorStatics
 	{
 		ConstructorHelpers::FObjectFinder<UStaticMesh> Plane;
 		FConstructorStatics()
-			: Plane(TEXT("/Engine/BasicShapes/Plane"))
+			: Plane(TEXT("/MediaPlate/SM_MediaPlateScreen"))
 		{}
 	};
 
