MediaPlate: Fixed up double entry in the playlist when dragging a media source into the viewport to create a media plate.
#rnx

diff --git a/Engine/Plugins/Media/MediaPlate/Source/MediaPlateEditor/Private/ActorFactoryMediaPlate.cpp b/Engine/Plugins/Media/MediaPlate/Source/MediaPlateEditor/Private/ActorFactoryMediaPlate.cpp
--- a/Engine/Plugins/Media/MediaPlate/Source/MediaPlateEditor/Private/ActorFactoryMediaPlate.cpp
+++ b/Engine/Plugins/Media/MediaPlate/Source/MediaPlateEditor/Private/ActorFactoryMediaPlate.cpp
@@ -61,7 +61,11 @@ void UActorFactoryMediaPlate::SetUpActor(UObject* Asset, AActor* Actor)
 			UMediaPlaylist* Playlist = MediaPlate->MediaPlateComponent->MediaPlaylist;
 			if (Playlist != nullptr)
 			{
-				Playlist->Add(MediaSource);
+				// This gets called twice so only add the media source once.
+				if (Playlist->Num() == 0)
+				{
+					Playlist->Add(MediaSource);
+				}
 			}
 		}
 	}
