MediaPlate: Added playlist group as a better place for the playlist.
#rnx

diff --git a/Engine/Plugins/Media/MediaPlate/Source/MediaPlate/Public/MediaPlateComponent.h b/Engine/Plugins/Media/MediaPlate/Source/MediaPlate/Public/MediaPlateComponent.h
--- a/Engine/Plugins/Media/MediaPlate/Source/MediaPlate/Public/MediaPlateComponent.h
+++ b/Engine/Plugins/Media/MediaPlate/Source/MediaPlate/Public/MediaPlateComponent.h
@@ -81,7 +81,7 @@ public:
 	TObjectPtr<UMediaSoundComponent> SoundComponent;
 
 	/** What media playlist to play. */
-	UPROPERTY(EditAnywhere, Category = "MediaPlate Playlist")
+	UPROPERTY(EditAnywhere, Category = "MediaPlate")
 	TObjectPtr<UMediaPlaylist> MediaPlaylist;
 
 	/** Enable smart caching for image sequences. */
diff --git a/Engine/Plugins/Media/MediaPlate/Source/MediaPlateEditor/Private/MediaPlateCustomization.cpp b/Engine/Plugins/Media/MediaPlate/Source/MediaPlateEditor/Private/MediaPlateCustomization.cpp
--- a/Engine/Plugins/Media/MediaPlate/Source/MediaPlateEditor/Private/MediaPlateCustomization.cpp
+++ b/Engine/Plugins/Media/MediaPlate/Source/MediaPlateEditor/Private/MediaPlateCustomization.cpp
@@ -6,6 +6,7 @@
 #include "DetailLayoutBuilder.h"
 #include "DetailWidgetRow.h"
 #include "EditorStyleSet.h"
+#include "IDetailGroup.h"
 #include "MediaPlate.h"
 #include "MediaPlateComponent.h"
 #include "MediaPlateEditorModule.h"
@@ -60,10 +61,15 @@ void FMediaPlateCustomization::CustomizeDetails(IDetailLayoutBuilder& DetailBuil
 	// Set media path.
 	UpdateMediaPath();
 	
-	IDetailCategoryBuilder& MediaPlaylistCategory = DetailBuilder.EditCategory("MediaPlate Playlist");
+	// Create playlist group.
+	IDetailGroup& PlaylistGroup = MediaPlateCategory.AddGroup(TEXT("Playlist"),
+		LOCTEXT("Playlist", "Playlist"));
+	TSharedRef<IPropertyHandle> PropertyHandle = DetailBuilder.GetProperty(
+		GET_MEMBER_NAME_CHECKED(UMediaPlateComponent, MediaPlaylist));
+	PlaylistGroup.HeaderProperty(PropertyHandle);
 
 	// Add media source.
-	MediaPlaylistCategory.AddCustomRow(LOCTEXT("MediaPlateMediaSource", "MediaPlate Media Source"))
+	PlaylistGroup.AddWidgetRow()
 		.NameContent()
 		[
 			SNew(STextBlock)
@@ -81,7 +87,7 @@ void FMediaPlateCustomization::CustomizeDetails(IDetailLayoutBuilder& DetailBuil
 
 	// Add media path.
 	FString FileTypeFilter = TEXT("All files (*.*)|*.*");
-	MediaPlaylistCategory.AddCustomRow(LOCTEXT("MediaPlateMediaPath", "MediaPlate Media Path"))
+	PlaylistGroup.AddWidgetRow()
 		.NameContent()
 		[
 			SNew(STextBlock)
