MediaPlate: Added scroll box to media details widget.
#rnx

diff --git a/Engine/Plugins/Media/MediaPlate/Source/MediaPlateEditor/Private/Widgets/SMediaPlateEditorMediaDetails.cpp b/Engine/Plugins/Media/MediaPlate/Source/MediaPlateEditor/Private/Widgets/SMediaPlateEditorMediaDetails.cpp
--- a/Engine/Plugins/Media/MediaPlate/Source/MediaPlateEditor/Private/Widgets/SMediaPlateEditorMediaDetails.cpp
+++ b/Engine/Plugins/Media/MediaPlate/Source/MediaPlateEditor/Private/Widgets/SMediaPlateEditorMediaDetails.cpp
@@ -7,6 +7,7 @@
 #include "MediaTexture.h"
 #include "Modules/ModuleManager.h"
 #include "PropertyEditorModule.h"
+#include "Widgets/Layout/SScrollBox.h"
 #include "Widgets/Text/STextBlock.h"
 
 #define LOCTEXT_NAMESPACE "SMediaPlateEditorMediaDetails"
@@ -20,80 +21,86 @@ void SMediaPlateEditorMediaDetails::Construct(const FArguments& InArgs, UMediaPl
 
 	ChildSlot
 		[
-			SNew(SHorizontalBox)
+			SNew(SScrollBox)
 
-			// Left side.
-			+ SHorizontalBox::Slot()
+			// Add details.
+			+ SScrollBox::Slot()
 				[
-					SNew(SVerticalBox)
+					SNew(SHorizontalBox)
 
-					// Resolution.
-					+ SVerticalBox::Slot()
-						.AutoHeight()
-						.VAlign(VAlign_Center)
-						.Padding(4.0f)
+					// Left side.
+					+ SHorizontalBox::Slot()
 						[
-							SAssignNew(ResolutionText, STextBlock)
+							SNew(SVerticalBox)
+
+							// Resolution.
+							+ SVerticalBox::Slot()
+								.AutoHeight()
+								.VAlign(VAlign_Center)
+								.Padding(4.0f)
+								[
+									SAssignNew(ResolutionText, STextBlock)
+								]
+
+							// Frame rate.
+							+ SVerticalBox::Slot()
+								.AutoHeight()
+								.VAlign(VAlign_Center)
+								.Padding(4.0f)
+								[
+									SAssignNew(FrameRateText, STextBlock)
+								]
+
+							// Resource size.
+							+ SVerticalBox::Slot()
+								.AutoHeight()
+								.VAlign(VAlign_Center)
+								.Padding(4.0f)
+								[
+									SAssignNew(ResourceSizeText, STextBlock)
+								]
 						]
 
-					// Frame rate.
-					+ SVerticalBox::Slot()
-						.AutoHeight()
-						.VAlign(VAlign_Center)
-						.Padding(4.0f)
+					// Right side.
+					+ SHorizontalBox::Slot()
 						[
-							SAssignNew(FrameRateText, STextBlock)
-						]
-
-					// Resource size.
-					+ SVerticalBox::Slot()
-						.AutoHeight()
-						.VAlign(VAlign_Center)
-						.Padding(4.0f)
-						[
-							SAssignNew(ResourceSizeText, STextBlock)
-						]
-				]
-
-			// Right side.
-			+ SHorizontalBox::Slot()
-				[
-					SNew(SVerticalBox)
-
-					// Method.
-					+ SVerticalBox::Slot()
-						.AutoHeight()
-						.VAlign(VAlign_Center)
-						.Padding(4.0f)
-						[
-							SAssignNew(MethodText, STextBlock)
-						]
-
-					// Format.
-					+ SVerticalBox::Slot()
-						.AutoHeight()
-						.VAlign(VAlign_Center)
-						.Padding(4.0f)
-						[
-							SAssignNew(FormatText, STextBlock)
-						]
-
-					// LOD bias.
-					+ SVerticalBox::Slot()
-						.AutoHeight()
-						.VAlign(VAlign_Center)
-						.Padding(4.0f)
-						[
-							SAssignNew(LODBiasText, STextBlock)
-						]
-
-					// Num mips.
-					+ SVerticalBox::Slot()
-						.AutoHeight()
-						.VAlign(VAlign_Center)
-						.Padding(4.0f)
-						[
-							SAssignNew(NumMipsText, STextBlock)
+							SNew(SVerticalBox)
+
+							// Method.
+							+ SVerticalBox::Slot()
+								.AutoHeight()
+								.VAlign(VAlign_Center)
+								.Padding(4.0f)
+								[
+									SAssignNew(MethodText, STextBlock)
+								]
+
+							// Format.
+							+ SVerticalBox::Slot()
+								.AutoHeight()
+								.VAlign(VAlign_Center)
+								.Padding(4.0f)
+								[
+									SAssignNew(FormatText, STextBlock)
+								]
+
+							// LOD bias.
+							+ SVerticalBox::Slot()
+								.AutoHeight()
+								.VAlign(VAlign_Center)
+								.Padding(4.0f)
+								[
+									SAssignNew(LODBiasText, STextBlock)
+								]
+
+							// Num mips.
+							+ SVerticalBox::Slot()
+								.AutoHeight()
+								.VAlign(VAlign_Center)
+								.Padding(4.0f)
+								[
+									SAssignNew(NumMipsText, STextBlock)
+								]
 						]
 				]
 
