ImgMedia: We simply share the EMediaTextureVisibleMipsTiles enum between MediaTextureTracker and MediaPlateComponent.

diff --git a/Engine/Plugins/Media/MediaPlate/Source/MediaPlate/Private/MediaPlateComponent.cpp b/Engine/Plugins/Media/MediaPlate/Source/MediaPlate/Private/MediaPlateComponent.cpp
--- a/Engine/Plugins/Media/MediaPlate/Source/MediaPlate/Private/MediaPlateComponent.cpp
+++ b/Engine/Plugins/Media/MediaPlate/Source/MediaPlate/Private/MediaPlateComponent.cpp
@@ -37,7 +37,7 @@ UMediaPlateComponent::UMediaPlateComponent(const FObjectInitializer& ObjectIniti
 	MediaPlaylist = CreateDefaultSubobject<UMediaPlaylist>(MediaPlaylistName);
 
 	// Default to plane since AMediaPlate defaults to SM_MediaPlateScreen
-	VisibleMipsTilesCalculations = EMediaPlateVisibleMipsTiles::Plane;
+	VisibleMipsTilesCalculations = EMediaTextureVisibleMipsTiles::Plane;
 }
 
 void UMediaPlateComponent::OnRegister()
@@ -124,7 +124,7 @@ void UMediaPlateComponent::RegisterWithMediaTextureTracker()
 	MediaTextureTrackerObject = MakeShared<FMediaTextureTrackerObject, ESPMode::ThreadSafe>();
 	MediaTextureTrackerObject->Object = GetOwner();
 	MediaTextureTrackerObject->MipMapLODBias = 0.0f;
-	MediaTextureTrackerObject->VisibleMipsTilesCalculations = static_cast<EMediaTextureVisibleMipsTiles>(VisibleMipsTilesCalculations);
+	MediaTextureTrackerObject->VisibleMipsTilesCalculations = VisibleMipsTilesCalculations;
 
 	// Add our texture.
 	TObjectPtr<UMediaTexture> MediaTexture = GetMediaTexture();
@@ -214,7 +214,7 @@ void UMediaPlateComponent::PostEditChangeProperty(FPropertyChangedEvent& Propert
 	{
 		if (MediaTextureTrackerObject != nullptr)
 		{
-			MediaTextureTrackerObject->VisibleMipsTilesCalculations = static_cast<EMediaTextureVisibleMipsTiles>(VisibleMipsTilesCalculations);
+			MediaTextureTrackerObject->VisibleMipsTilesCalculations = VisibleMipsTilesCalculations;
 		}
 	}
 }
diff --git a/Engine/Plugins/Media/MediaPlate/Source/MediaPlate/Public/MediaPlateComponent.h b/Engine/Plugins/Media/MediaPlate/Source/MediaPlate/Public/MediaPlateComponent.h
--- a/Engine/Plugins/Media/MediaPlate/Source/MediaPlate/Public/MediaPlateComponent.h
+++ b/Engine/Plugins/Media/MediaPlate/Source/MediaPlate/Public/MediaPlateComponent.h
@@ -15,18 +15,6 @@ class UMediaPlaylist;
 class UMediaSoundComponent;
 class UMediaSource;
 class UMediaTexture;
-struct FMediaTextureTrackerObject;
-
-/**
- * List of mesh types used for visible mips and tiles calculations, matching EMediaTextureVisibleMipsTiles.
- */
-UENUM()
-enum class EMediaPlateVisibleMipsTiles : uint8
-{
-	Disabled = 0,
-	Plane = 1,
-	Sphere
-};
 
 
 /**
@@ -114,7 +102,7 @@ public:
 	 * Use /Engine/BasicShapes/Sphere for the sphere.
 	 */
 	UPROPERTY(EditAnywhere, AdvancedDisplay, Category = "MediaPlate")
-	EMediaPlateVisibleMipsTiles VisibleMipsTilesCalculations;
+	EMediaTextureVisibleMipsTiles VisibleMipsTilesCalculations;
 
 	/**
 	 * Adds our media texture to the media texture tracker.
diff --git a/Engine/Source/Runtime/MediaAssets/Public/MediaTextureTracker.h b/Engine/Source/Runtime/MediaAssets/Public/MediaTextureTracker.h
--- a/Engine/Source/Runtime/MediaAssets/Public/MediaTextureTracker.h
+++ b/Engine/Source/Runtime/MediaAssets/Public/MediaTextureTracker.h
@@ -3,13 +3,17 @@
 #pragma once
 
 #include "CoreMinimal.h"
+#include "UObject/ObjectMacros.h"
 #include "UObject/ObjectPtr.h"
 #include "UObject/WeakObjectPtr.h"
 #include "UObject/WeakObjectPtrTemplates.h"
 
+#include "MediaTextureTracker.generated.h"
+
 class AActor;
 class UMediaTexture;
 
+UENUM()
 enum class EMediaTextureVisibleMipsTiles : uint8
 {
 	None = 0,
