MediaPlate: Can change the aspect ratio of plane meshes in the editor.

#rnx

diff --git a/Engine/Plugins/Media/MediaPlate/Source/MediaPlateEditor/Private/MediaPlateCustomization.cpp b/Engine/Plugins/Media/MediaPlate/Source/MediaPlateEditor/Private/MediaPlateCustomization.cpp
--- a/Engine/Plugins/Media/MediaPlate/Source/MediaPlateEditor/Private/MediaPlateCustomization.cpp
+++ b/Engine/Plugins/Media/MediaPlate/Source/MediaPlateEditor/Private/MediaPlateCustomization.cpp
@@ -5,8 +5,9 @@
 #include "DetailCategoryBuilder.h"
 #include "DetailLayoutBuilder.h"
 #include "DetailWidgetRow.h"
-#include "Styling/AppStyle.h"
+#include "Framework/MultiBox/MultiBoxBuilder.h"
 #include "IDetailGroup.h"
+#include "LevelEditorViewport.h"
 #include "MediaPlate.h"
 #include "MediaPlateComponent.h"
 #include "MediaPlateEditorModule.h"
@@ -14,9 +15,13 @@
 #include "MediaPlaylist.h"
 #include "MediaSource.h"
 #include "PropertyCustomizationHelpers.h"
+#include "Styling/AppStyle.h"
 #include "Subsystems/AssetEditorSubsystem.h"
 #include "Widgets/Input/SButton.h"
+#include "Widgets/Input/SComboButton.h"
 #include "Widgets/Input/SFilePathPicker.h"
+#include "Widgets/Input/SNumericEntryBox.h"
+#include "Widgets/Input/SSegmentedControl.h"
 #include "Widgets/Images/SImage.h"
 
 #define LOCTEXT_NAMESPACE "FMediaPlateCustomization"
@@ -55,11 +60,15 @@ void FMediaPlateCustomization::CustomizeDetails(IDetailLayoutBuilder& DetailBuil
 		if (MediaPlate.IsValid())
 		{
 			MediaPlatesList.Add(MediaPlate);
+			MeshMode = MediaPlate->VisibleMipsTilesCalculations;
 		}
 	}
 
 	// Set media path.
 	UpdateMediaPath();
+
+	// Add mesh customization.
+	AddMeshCustomization(MediaPlateCategory);
 	
 	// Create playlist group.
 	IDetailGroup& PlaylistGroup = MediaPlateCategory.AddGroup(TEXT("Playlist"),
@@ -328,6 +337,183 @@ void FMediaPlateCustomization::CustomizeDetails(IDetailLayoutBuilder& DetailBuil
 	}
 }
 
+void FMediaPlateCustomization::AddMeshCustomization(IDetailCategoryBuilder& MediaPlateCategory)
+{
+	// Create detail group.
+	IDetailGroup& DetailGroup = MediaPlateCategory.AddGroup(TEXT("Mesh"),
+		LOCTEXT("Mesh", "Mesh"));
+
+	// Add radio buttons for mesh type.
+	DetailGroup.AddWidgetRow()
+		[
+			SNew(SSegmentedControl<EMediaTextureVisibleMipsTiles>)
+				.Value_Lambda([this]()
+				{
+					return MeshMode;
+				})
+				.OnValueChanged_Lambda([this](EMediaTextureVisibleMipsTiles Mode)
+				{
+					MeshMode = Mode;
+				})
+
+			+ SSegmentedControl<EMediaTextureVisibleMipsTiles>::Slot(EMediaTextureVisibleMipsTiles::Plane)
+				.Text(LOCTEXT("Plane", "Plane"))
+				.ToolTip(LOCTEXT("Plane_ToolTip",
+					"Select this if you want to use a standard plane for the mesh."))
+
+			+ SSegmentedControl<EMediaTextureVisibleMipsTiles>::Slot(EMediaTextureVisibleMipsTiles::Sphere)
+				.Text(LOCTEXT("Sphere", "Sphere"))
+				.ToolTip(LOCTEXT("Sphere_ToolTip",
+					"Select this if you want to use a spherical object for the mesh."))
+
+			+ SSegmentedControl<EMediaTextureVisibleMipsTiles>::Slot(EMediaTextureVisibleMipsTiles::None)
+				.Text(LOCTEXT("Custom", "Custom"))
+				.ToolTip(LOCTEXT("Custom_ToolTip",
+					"Select this if you want to provide your own mesh."))
+		];
+
+	// Visibility attributes.
+	TAttribute<EVisibility> MeshPlaneVisibility(this, &FMediaPlateCustomization::ShouldShowMeshPlaneWidgets);
+
+	// Add aspect ratio.
+	DetailGroup.AddWidgetRow()
+		.Visibility(MeshPlaneVisibility)
+		.NameContent()
+		[
+			SNew(STextBlock)
+				.Text(LOCTEXT("AspectRatio", "Aspect Ratio"))
+				.ToolTipText(LOCTEXT("AspectRatio_ToolTip",
+				"Sets the aspect ratio of the plane showing the media.\nChanging this will change the scale of the mesh component."))
+				.Font(IDetailLayoutBuilder::GetDetailFont())
+		]
+		.ValueContent()
+		[
+			SNew(SHorizontalBox)
+
+			// Presets button.
+			+ SHorizontalBox::Slot()
+				.AutoWidth()
+				[
+					SNew(SComboButton)
+						.OnGetMenuContent(this, &FMediaPlateCustomization::OnGetAspectRatios)
+						.ContentPadding(2)
+						.ButtonContent()
+						[
+							SNew(STextBlock)
+								.ToolTipText(LOCTEXT("Presets_ToolTip", "Select one of the presets for the aspect ratio."))
+								.Text(LOCTEXT("Presets", "Presets"))
+						]
+				]
+
+			// Numeric entry box.
+			+ SHorizontalBox::Slot()
+				.AutoWidth()
+				[
+					SNew(SNumericEntryBox<float>)
+						.Value(this, &FMediaPlateCustomization::GetAspectRatio)
+						.MinValue(0.0f)
+						.OnValueChanged(this, &FMediaPlateCustomization::SetAspectRatio)
+				]
+		];
+}
+
+EVisibility FMediaPlateCustomization::ShouldShowMeshPlaneWidgets() const
+{
+	return (MeshMode == EMediaTextureVisibleMipsTiles::Plane) ? EVisibility::Visible : EVisibility::Hidden;
+}
+
+TSharedRef<SWidget> FMediaPlateCustomization::OnGetAspectRatios()
+{
+	FMenuBuilder MenuBuilder(true, NULL);
+
+	FUIAction Set16x9Action(FExecuteAction::CreateSP(this, &FMediaPlateCustomization::SetAspectRatio, 16.0f / 9.0f));
+	MenuBuilder.AddMenuEntry(LOCTEXT("16x9", "16x9"), FText(), FSlateIcon(), Set16x9Action);
+
+	FUIAction Set16x10Action(FExecuteAction::CreateSP(this, &FMediaPlateCustomization::SetAspectRatio, 16.0f / 10.0f));
+	MenuBuilder.AddMenuEntry(LOCTEXT("16x10", "16x10"), FText(), FSlateIcon(), Set16x10Action);
+
+	FUIAction Set4x3Action(FExecuteAction::CreateSP(this, &FMediaPlateCustomization::SetAspectRatio, 4.0f / 3.0f));
+	MenuBuilder.AddMenuEntry(LOCTEXT("4x3", "4x3"), FText(), FSlateIcon(), Set4x3Action);
+
+	FUIAction Set1x1Action(FExecuteAction::CreateSP(this, &FMediaPlateCustomization::SetAspectRatio, 1.0f));
+	MenuBuilder.AddMenuEntry(LOCTEXT("1x1", "1x1"), FText(), FSlateIcon(), Set1x1Action);
+
+	return MenuBuilder.MakeWidget();
+}
+
+void FMediaPlateCustomization::SetAspectRatio(float AspectRatio)
+{
+	// Calculate required scale from the aspect ratio.
+	float Height = 1.0f;
+	if (AspectRatio != 0.0f)
+	{
+		Height = 1.0f / AspectRatio;
+	}
+	FVector Scale(1.0f, 1.0f, Height);
+
+	// Loop through all our objects.
+	for (const TWeakObjectPtr<UMediaPlateComponent>& MediaPlatePtr : MediaPlatesList)
+	{
+		UMediaPlateComponent* MediaPlate = MediaPlatePtr.Get();
+		if (MediaPlate != nullptr)
+		{
+			AActor* Owner = MediaPlate->GetOwner();
+			AMediaPlate* MediaPlateActor = Cast<AMediaPlate>(Owner);
+			if (MediaPlateActor != nullptr)
+			{
+				// Get the static mesh.
+				UStaticMeshComponent* StaticMeshComponent = MediaPlateActor->StaticMeshComponent;
+				if (StaticMeshComponent != nullptr)
+				{
+					// Update the scale.
+					StaticMeshComponent->SetRelativeScale3D(Scale);
+					StaticMeshComponent->MarkPackageDirty();
+				}
+			}
+		}
+	}
+
+	// Invalidate the viewport so we can see the mesh change.
+	if (GCurrentLevelEditingViewportClient != nullptr)
+	{
+		GCurrentLevelEditingViewportClient->Invalidate();
+	}
+}
+
+
+TOptional<float> FMediaPlateCustomization::GetAspectRatio() const
+{
+	// Loop through our objects.
+	for (const TWeakObjectPtr<UMediaPlateComponent>& MediaPlatePtr : MediaPlatesList)
+	{
+		UMediaPlateComponent* MediaPlate = MediaPlatePtr.Get();
+		if (MediaPlate != nullptr)
+		{
+			AActor* Owner = MediaPlate->GetOwner();
+			AMediaPlate* MediaPlateActor = Cast<AMediaPlate>(Owner);
+			if (MediaPlateActor != nullptr)
+			{
+				// Get the static mesh.
+				UStaticMeshComponent* StaticMeshComponent = MediaPlateActor->StaticMeshComponent;
+				if (StaticMeshComponent != nullptr)
+				{
+					// Calculate aspect ratio from the scale.
+					FVector Scale = StaticMeshComponent->GetComponentScale();
+					float AspectRatio = 0.0f;
+					if (Scale.Z != 0.0f)
+					{
+						AspectRatio = Scale.Y / Scale.Z;
+					}
+					return AspectRatio;
+				}
+			}
+			break;
+		}
+	}
+
+	return TOptional<float>();
+}
+
 FString FMediaPlateCustomization::GetMediaSourcePath() const
 {
 	FString Path;
diff --git a/Engine/Plugins/Media/MediaPlate/Source/MediaPlateEditor/Private/MediaPlateCustomization.h b/Engine/Plugins/Media/MediaPlate/Source/MediaPlateEditor/Private/MediaPlateCustomization.h
--- a/Engine/Plugins/Media/MediaPlate/Source/MediaPlateEditor/Private/MediaPlateCustomization.h
+++ b/Engine/Plugins/Media/MediaPlate/Source/MediaPlateEditor/Private/MediaPlateCustomization.h
@@ -6,8 +6,11 @@
 #include "Editor.h"
 #include "IDetailCustomization.h"
 #include "Input/Reply.h"
+#include "MediaTextureTracker.h"
 
+class IDetailCategoryBuilder;
 class IDetailLayoutBuilder;
+class SWidget;
 class UMediaPlateComponent;
 class UMediaPlayer;
 
@@ -40,6 +43,34 @@ private:
 	/** Stores the current value of the MediaPath property. */
 	FString MediaPath;
 
+	/** Whether we have a plane, sphere, etc. */
+	EMediaTextureVisibleMipsTiles MeshMode;
+
+	/**
+	 * Adds widgets for editing the mesh.
+	 */
+	void AddMeshCustomization(IDetailCategoryBuilder& MediaPlateCategory);
+
+	/**
+	 * Controls visibility for widgets for plane meshes.
+	 */
+	EVisibility ShouldShowMeshPlaneWidgets() const;
+
+	/**
+	 * Returns menu options for all aspect ratio presets.
+	 */
+	TSharedRef<SWidget> OnGetAspectRatios();
+
+	/**
+	 * Call this to set the aspect ratio.
+	 */
+	void SetAspectRatio(float AspectRatio);
+
+	/**
+	 * Call this to get the aspect ratio.
+	 */
+	TOptional<float> GetAspectRatio() const;
+
 	/**
 	 * Gets the object path for the media source object.
 	 */
