MediaPlate: Playlist will be created if the media source is changed but there is no playlist.
#rnx

diff --git a/Engine/Plugins/Media/MediaPlate/Source/MediaPlateEditor/Private/MediaPlateCustomization.cpp b/Engine/Plugins/Media/MediaPlate/Source/MediaPlateEditor/Private/MediaPlateCustomization.cpp
--- a/Engine/Plugins/Media/MediaPlate/Source/MediaPlateEditor/Private/MediaPlateCustomization.cpp
+++ b/Engine/Plugins/Media/MediaPlate/Source/MediaPlateEditor/Private/MediaPlateCustomization.cpp
@@ -549,19 +549,24 @@ void FMediaPlateCustomization::OnMediaSourceChanged(const FAssetData& AssetData)
 		UMediaPlateComponent* MediaPlate = MediaPlatePtr.Get();
 		if (MediaPlate != nullptr)
 		{
+			// Get playlist.
 			UMediaPlaylist* Playlist = MediaPlate->MediaPlaylist;
-			if (Playlist != nullptr)
+			if (Playlist == nullptr)
 			{
-				if (Playlist->Num() > 0)
-				{
-					Playlist->Replace(0, MediaSource);
-				}
-				else
-				{
-					Playlist->Add(MediaSource);
-				}
-				Playlist->MarkPackageDirty();
+				Playlist = NewObject<UMediaPlaylist>(MediaPlate);
+				MediaPlate->MediaPlaylist = Playlist;
+			}
+			
+			// Update playlist.
+			if (Playlist->Num() > 0)
+			{
+				Playlist->Replace(0, MediaSource);
+			}
+			else
+			{
+				Playlist->Add(MediaSource);
 			}
+			Playlist->MarkPackageDirty();
 		}
 	}
 
@@ -624,23 +629,27 @@ void FMediaPlateCustomization::HandleMediaPathPicked(const FString& PickedPath)
 			UMediaPlateComponent* MediaPlate = MediaPlatePtr.Get();
 			if (MediaPlate != nullptr)
 			{
+				// Get playlist.
 				UMediaPlaylist* Playlist = MediaPlate->MediaPlaylist;
-				if (Playlist != nullptr)
+				if (Playlist == nullptr)
 				{
-					// Create media source for this path.
-					UMediaSource* MediaSource = UMediaSource::SpawnMediaSourceForString(PickedPath, MediaPlate);
-					if (MediaSource != nullptr)
+					Playlist = NewObject<UMediaPlaylist>(MediaPlate);
+					MediaPlate->MediaPlaylist = Playlist;
+				}
+				
+				// Create media source for this path.
+				UMediaSource* MediaSource = UMediaSource::SpawnMediaSourceForString(PickedPath, MediaPlate);
+				if (MediaSource != nullptr)
+				{
+					if (Playlist->Num() > 0)
 					{
-						if (Playlist->Num() > 0)
-						{
-							Playlist->Replace(0, MediaSource);
-						}
-						else
-						{
-							Playlist->Add(MediaSource);
-						}
-						Playlist->MarkPackageDirty();
+						Playlist->Replace(0, MediaSource);
+					}
+					else
+					{
+						Playlist->Add(MediaSource);
 					}
+					Playlist->MarkPackageDirty();
 				}
 			}
 		}
