MediaPlate: Media path updates when the playlist changes in the editor.

#rnx

diff --git a/Engine/Plugins/Media/MediaPlate/Source/MediaPlateEditor/Private/MediaPlateCustomization.cpp b/Engine/Plugins/Media/MediaPlate/Source/MediaPlateEditor/Private/MediaPlateCustomization.cpp
--- a/Engine/Plugins/Media/MediaPlate/Source/MediaPlateEditor/Private/MediaPlateCustomization.cpp
+++ b/Engine/Plugins/Media/MediaPlate/Source/MediaPlateEditor/Private/MediaPlateCustomization.cpp
@@ -76,6 +76,8 @@ void FMediaPlateCustomization::CustomizeDetails(IDetailLayoutBuilder& DetailBuil
 	TSharedRef<IPropertyHandle> PropertyHandle = DetailBuilder.GetProperty(
 		GET_MEMBER_NAME_CHECKED(UMediaPlateComponent, MediaPlaylist));
 	PlaylistGroup.HeaderProperty(PropertyHandle);
+	PropertyHandle->SetOnPropertyValueChanged(FSimpleDelegate::CreateSP(this,
+		&FMediaPlateCustomization::OnPlaylistChanged));
 
 	// Add media source.
 	PlaylistGroup.AddWidgetRow()
@@ -540,6 +542,12 @@ FString FMediaPlateCustomization::GetMediaSourcePath() const
 	return Path;
 }
 
+void FMediaPlateCustomization::OnPlaylistChanged()
+{
+	StopMediaPlates();
+	UpdateMediaPath();
+}
+
 void FMediaPlateCustomization::OnMediaSourceChanged(const FAssetData& AssetData)
 {
 	// Update the playlist with the new media source.
diff --git a/Engine/Plugins/Media/MediaPlate/Source/MediaPlateEditor/Private/MediaPlateCustomization.h b/Engine/Plugins/Media/MediaPlate/Source/MediaPlateEditor/Private/MediaPlateCustomization.h
--- a/Engine/Plugins/Media/MediaPlate/Source/MediaPlateEditor/Private/MediaPlateCustomization.h
+++ b/Engine/Plugins/Media/MediaPlate/Source/MediaPlateEditor/Private/MediaPlateCustomization.h
@@ -76,6 +76,11 @@ private:
 	 */
 	FString GetMediaSourcePath() const;
 
+	/**
+	 * Called when the playlist changes.
+	 */
+	void OnPlaylistChanged();
+
 	/**
 	 * Called when the media source widget changes.
 	 */
