MediaPlate: Creates a sphere mesh when sphere mode is selected in the editor.

#rnx

diff --git a/Engine/Plugins/Media/MediaPlate/MediaPlate.uplugin b/Engine/Plugins/Media/MediaPlate/MediaPlate.uplugin
--- a/Engine/Plugins/Media/MediaPlate/MediaPlate.uplugin
+++ b/Engine/Plugins/Media/MediaPlate/MediaPlate.uplugin
@@ -39,6 +39,11 @@
 			"Name": "MediaPlayerEditor",
 			"TargetAllowList": [ "Editor" ],
 			"Enabled": true
-		}
+		},
+		{
+			"Name": "MeshModelingToolset",
+			"TargetAllowList": [ "Editor" ],
+			"Enabled": true
+		},
 	]
 }
diff --git a/Engine/Plugins/Media/MediaPlate/Source/MediaPlateEditor/MediaPlateEditor.Build.cs b/Engine/Plugins/Media/MediaPlate/Source/MediaPlateEditor/MediaPlateEditor.Build.cs
--- a/Engine/Plugins/Media/MediaPlate/Source/MediaPlateEditor/MediaPlateEditor.Build.cs
+++ b/Engine/Plugins/Media/MediaPlate/Source/MediaPlateEditor/MediaPlateEditor.Build.cs
@@ -22,12 +22,15 @@
 				"DesktopWidgets",
 				"EditorStyle",
 				"Engine",
+				"GeometryCore",
 				"InputCore",
 				"MediaAssets",
 				"MediaCompositing",
 				"MediaCompositingEditor",
 				"MediaPlate",
 				"MediaPlayerEditor",
+				"ModelingComponentsEditorOnly",
+				"ModelingComponents",
 				"MovieScene",
 				"PropertyEditor",
 				"Sequencer",
diff --git a/Engine/Plugins/Media/MediaPlate/Source/MediaPlateEditor/Private/MediaPlateCustomization.cpp b/Engine/Plugins/Media/MediaPlate/Source/MediaPlateEditor/Private/MediaPlateCustomization.cpp
--- a/Engine/Plugins/Media/MediaPlate/Source/MediaPlateEditor/Private/MediaPlateCustomization.cpp
+++ b/Engine/Plugins/Media/MediaPlate/Source/MediaPlateEditor/Private/MediaPlateCustomization.cpp
@@ -353,10 +353,7 @@
 				{
 					return MeshMode;
 				})
-				.OnValueChanged_Lambda([this](EMediaTextureVisibleMipsTiles Mode)
-				{
-					MeshMode = Mode;
-				})
+				.OnValueChanged(this, &FMediaPlateCustomization::SetMeshMode)
 
 			+ SSegmentedControl<EMediaTextureVisibleMipsTiles>::Slot(EMediaTextureVisibleMipsTiles::Plane)
 				.Text(LOCTEXT("Plane", "Plane"))
@@ -424,6 +421,33 @@
 	return (MeshMode == EMediaTextureVisibleMipsTiles::Plane) ? EVisibility::Visible : EVisibility::Hidden;
 }
 
+void FMediaPlateCustomization::SetMeshMode(EMediaTextureVisibleMipsTiles InMode)
+{
+	if (MeshMode != InMode)
+	{
+		MeshMode = InMode;
+		for (const TWeakObjectPtr<UMediaPlateComponent>& MediaPlatePtr : MediaPlatesList)
+		{
+			UMediaPlateComponent* MediaPlate = MediaPlatePtr.Get();
+			if (MediaPlate != nullptr)
+			{
+				// Update the setting in the media plate.
+				MediaPlate->VisibleMipsTilesCalculations = MeshMode;
+
+				// Set the appropriate mesh.
+				if (MeshMode == EMediaTextureVisibleMipsTiles::Plane)
+				{
+					MeshCustomization.SetPlaneMesh(MediaPlate);
+				}
+				else if (MeshMode == EMediaTextureVisibleMipsTiles::Sphere)
+				{
+					MeshCustomization.SetSphereMesh(MediaPlate);
+				}
+			}
+		}
+	}
+}
+
 TSharedRef<SWidget> FMediaPlateCustomization::OnGetAspectRatios()
 {
 	FMenuBuilder MenuBuilder(true, NULL);
diff --git a/Engine/Plugins/Media/MediaPlate/Source/MediaPlateEditor/Private/MediaPlateCustomization.h b/Engine/Plugins/Media/MediaPlate/Source/MediaPlateEditor/Private/MediaPlateCustomization.h
--- a/Engine/Plugins/Media/MediaPlate/Source/MediaPlateEditor/Private/MediaPlateCustomization.h
+++ b/Engine/Plugins/Media/MediaPlate/Source/MediaPlateEditor/Private/MediaPlateCustomization.h
@@ -6,6 +6,7 @@
 #include "Editor.h"
 #include "IDetailCustomization.h"
 #include "Input/Reply.h"
+#include "MediaPlateCustomizationMesh.h"
 #include "MediaTextureTracker.h"
 
 class IDetailCategoryBuilder;
@@ -46,6 +47,9 @@
 	/** Whether we have a plane, sphere, etc. */
 	EMediaTextureVisibleMipsTiles MeshMode;
 
+	/** Handles mesh stuff. */
+	FMediaPlateCustomizationMesh MeshCustomization;
+
 	/**
 	 * Adds widgets for editing the mesh.
 	 */
@@ -57,6 +61,11 @@
 	EVisibility ShouldShowMeshPlaneWidgets() const;
 
 	/**
+	 * Call this to switch between planes, spheres, etc.
+	 */
+	void SetMeshMode(EMediaTextureVisibleMipsTiles InMode);
+
+	/**
 	 * Returns menu options for all aspect ratio presets.
 	 */
 	TSharedRef<SWidget> OnGetAspectRatios();
diff --git a/Engine/Plugins/Media/MediaPlate/Source/MediaPlateEditor/Private/MediaPlateCustomizationMesh.cpp b/Engine/Plugins/Media/MediaPlate/Source/MediaPlateEditor/Private/MediaPlateCustomizationMesh.cpp
new file mode 100644
--- /dev/null
+++ b/Engine/Plugins/Media/MediaPlate/Source/MediaPlateEditor/Private/MediaPlateCustomizationMesh.cpp
@@ -0,0 +1,147 @@
+// Copyright Epic Games, Inc. All Rights Reserved.
+
+#include "MediaPlateCustomizationMesh.h"
+
+#include "AssetRegistry/AssetRegistryModule.h"
+#include "AssetUtils/CreateStaticMeshUtil.h"
+#include "DynamicMesh/DynamicMesh3.h"
+#include "Generators/SphereGenerator.h"
+#include "LevelEditorViewport.h"
+#include "MediaPlate.h"
+#include "MediaPlateComponent.h"
+
+#define LOCTEXT_NAMESPACE "FMediaPlateCustomizationMesh"
+
+void FMediaPlateCustomizationMesh::SetPlaneMesh(UMediaPlateComponent* MediaPlate)
+{
+	// Get static mesh component.
+	UStaticMeshComponent* StaticMeshComponent = nullptr;
+	if (MediaPlate != nullptr)
+	{
+		AActor* Owner = MediaPlate->GetOwner();
+		AMediaPlate* MediaPlateActor = Cast<AMediaPlate>(Owner);
+		if (MediaPlateActor != nullptr)
+		{
+			StaticMeshComponent = MediaPlateActor->StaticMeshComponent;
+		}
+	}
+
+	if (StaticMeshComponent != nullptr)
+	{
+		// Get plane mesh.
+		UStaticMesh* Mesh = Cast<UStaticMesh>(StaticLoadObject(UStaticMesh::StaticClass(), NULL,
+			TEXT("/MediaPlate/SM_MediaPlateScreen")));
+
+		// Apply this mesh.
+		SetMesh(StaticMeshComponent, Mesh);
+	}
+}
+
+void FMediaPlateCustomizationMesh::SetSphereMesh(UMediaPlateComponent* MediaPlate)
+{
+	// Get static mesh component.
+	UStaticMeshComponent* StaticMeshComponent = nullptr;
+	if (MediaPlate != nullptr)
+	{
+		AActor* Owner = MediaPlate->GetOwner();
+		AMediaPlate* MediaPlateActor = Cast<AMediaPlate>(Owner);
+		if (MediaPlateActor != nullptr)
+		{
+			StaticMeshComponent = MediaPlateActor->StaticMeshComponent;
+		}
+	}
+
+	if (StaticMeshComponent != nullptr)
+	{
+		// Do we already have this mesh?
+		FString AssetPath = GetAssetPath();
+		UStaticMesh* StaticMesh = Cast<UStaticMesh>(StaticLoadObject(UStaticMesh::StaticClass(), NULL,
+			*AssetPath));
+		if (StaticMesh == nullptr)
+		{
+			// Create mesh.
+			FDynamicMesh3 NewMesh;
+			GenerateSphereMesh(&NewMesh);
+
+			// Create asset.
+			StaticMesh = CreateStaticMeshAsset(&NewMesh, AssetPath);
+		}
+
+		// Apply mesh.
+		SetMesh(StaticMeshComponent, StaticMesh);
+	}
+}
+
+void FMediaPlateCustomizationMesh::SetMesh(UStaticMeshComponent* StaticMeshComponent, UStaticMesh* Mesh)
+{
+	if (StaticMeshComponent != nullptr)
+	{
+		// Apply mesh to component.
+		StaticMeshComponent->SetStaticMesh(Mesh);
+		StaticMeshComponent->SetRelativeScale3D(FVector::OneVector);
+
+		// Call PostEditChangeProperty so it updates properly.
+		FProperty* StaticMeshProperty = FindFieldChecked<FProperty>(UStaticMeshComponent::StaticClass(), "StaticMesh");
+		FPropertyChangedEvent PropertyEvent(StaticMeshProperty);
+		StaticMeshComponent->PostEditChangeProperty(PropertyEvent);
+
+		// Invalidate the viewport so we can see the mesh change.
+		if (GCurrentLevelEditingViewportClient != nullptr)
+		{
+			GCurrentLevelEditingViewportClient->Invalidate();
+		}
+	}
+}
+
+void FMediaPlateCustomizationMesh::GenerateSphereMesh(FDynamicMesh3* OutMesh)
+{
+	UE::Geometry::FSphereGenerator SphereGen;
+	SphereGen.Radius = 50.0f;
+	SphereGen.NumTheta = 17;
+	SphereGen.NumPhi = 17;
+	SphereGen.bPolygroupPerQuad = false;
+	SphereGen.Generate();
+
+	OutMesh->Copy(&SphereGen);
+}
+
+UStaticMesh* FMediaPlateCustomizationMesh::CreateStaticMeshAsset(FDynamicMesh3* Mesh, const FString& AssetPath)
+{
+	UStaticMesh* NewStaticMesh = nullptr;
+
+	UE::AssetUtils::FStaticMeshAssetOptions AssetOptions;
+
+	AssetOptions.NewAssetPath = AssetPath;
+	AssetOptions.NumSourceModels = 1;
+	
+	AssetOptions.bEnableRecomputeNormals = false;
+	AssetOptions.bEnableRecomputeTangents = false;
+	AssetOptions.bGenerateNaniteEnabledMesh = false;
+	AssetOptions.NaniteProxyTrianglePercent = 100.0f;
+
+	AssetOptions.bCreatePhysicsBody = true;;
+	AssetOptions.CollisionType = ECollisionTraceFlag::CTF_UseComplexAsSimple;
+
+	AssetOptions.SourceMeshes.DynamicMeshes.Add(Mesh);
+
+	UE::AssetUtils::FStaticMeshResults ResultData;
+	UE::AssetUtils::ECreateStaticMeshResult AssetResult = UE::AssetUtils::CreateStaticMeshAsset(AssetOptions, ResultData);
+
+	if (AssetResult == UE::AssetUtils::ECreateStaticMeshResult::Ok)
+	{
+		NewStaticMesh = ResultData.StaticMesh;
+		NewStaticMesh->MarkPackageDirty();
+		FAssetRegistryModule::AssetCreated(NewStaticMesh);
+	}
+
+	return NewStaticMesh;
+}
+
+FString FMediaPlateCustomizationMesh::GetAssetPath()
+{
+	FString AssetPath = "/Game/_MediaPlate/Sphere";
+	
+	return AssetPath;
+}
+
+#undef LOCTEXT_NAMESPACE
diff --git a/Engine/Plugins/Media/MediaPlate/Source/MediaPlateEditor/Private/MediaPlateCustomizationMesh.h b/Engine/Plugins/Media/MediaPlate/Source/MediaPlateEditor/Private/MediaPlateCustomizationMesh.h
new file mode 100644
--- /dev/null
+++ b/Engine/Plugins/Media/MediaPlate/Source/MediaPlateEditor/Private/MediaPlateCustomizationMesh.h
@@ -0,0 +1,51 @@
+// Copyright Epic Games, Inc. All Rights Reserved.
+
+#pragma once
+
+#include "GeometryBase.h"
+
+PREDECLARE_USE_GEOMETRY_CLASS(FDynamicMesh3);
+
+class UMediaPlateComponent;
+class UStaticMesh;
+class UStaticMeshComponent;
+
+/**
+ * Implements mesh handling for FMediaPlateCustomization.
+ */
+class FMediaPlateCustomizationMesh
+{
+public:
+	/**
+	 * Call this to use a plane mesh on a media plate.
+	 */
+	void SetPlaneMesh(UMediaPlateComponent* MediaPlate);
+
+	/**
+	 * Call this to use a sphere mesh on a media plate.
+	*/
+	void SetSphereMesh(UMediaPlateComponent* MediaPlate);
+
+private:
+
+	/**
+	 * Generic function to set a mesh on a StaticMeshComponent.
+	 */
+	void SetMesh(UStaticMeshComponent* StaticMeshComponent, UStaticMesh* Mesh);
+
+	/**
+	 * Creates a sphere mesh.
+	 */
+	void GenerateSphereMesh(FDynamicMesh3* OutMesh);
+
+	/**
+	 * Call this to create a static mesh asset for a mesh.
+	 */
+	UStaticMesh* CreateStaticMeshAsset(FDynamicMesh3* Mesh, const FString& AssetPath);
+
+	/**
+	 * Gets the path for the mesh asset.
+	 */
+	FString GetAssetPath();
+};
+
