MediaPlate: Propagate VisibleMipsTilesCalculations change through the tracker and player objects, restarting play if necessary. This prevents the incorrect calculations from occuring if such a change occurs during playback in editor.

#rnx

diff --git a/Engine/Plugins/Media/MediaPlate/Source/MediaPlate/Private/MediaPlateComponent.cpp b/Engine/Plugins/Media/MediaPlate/Source/MediaPlate/Private/MediaPlateComponent.cpp
--- a/Engine/Plugins/Media/MediaPlate/Source/MediaPlate/Private/MediaPlateComponent.cpp
+++ b/Engine/Plugins/Media/MediaPlate/Source/MediaPlate/Private/MediaPlateComponent.cpp
@@ -174,6 +174,25 @@ bool UMediaPlateComponent::PlayMediaSource(UMediaSource* InMediaSource)
 
 #if WITH_EDITOR
 
+void UMediaPlateComponent::OnVisibleMipsTilesCalculationsChange()
+{
+	if (MediaTextureTrackerObject != nullptr)
+	{
+		MediaTextureTrackerObject->VisibleMipsTilesCalculations = VisibleMipsTilesCalculations;
+
+		// Propagate the change by restarting the player if it is currently playing.
+		TObjectPtr<UMediaPlayer> MediaPlayer = GetMediaPlayer();
+		if (MediaPlayer != nullptr)
+		{
+			if (MediaPlayer->IsPlaying())
+			{
+				MediaPlayer->Close();
+				Play();
+			}
+		}
+	}
+}
+
 void UMediaPlateComponent::PostEditChangeProperty(FPropertyChangedEvent& PropertyChangedEvent)
 {
 	// Has bEnableAudiio changed?
@@ -212,10 +231,7 @@ void UMediaPlateComponent::PostEditChangeProperty(FPropertyChangedEvent& Propert
 	}
 	else if (PropertyChangedEvent.GetPropertyName() == GET_MEMBER_NAME_CHECKED(ThisClass, VisibleMipsTilesCalculations))
 	{
-		if (MediaTextureTrackerObject != nullptr)
-		{
-			MediaTextureTrackerObject->VisibleMipsTilesCalculations = VisibleMipsTilesCalculations;
-		}
+		OnVisibleMipsTilesCalculationsChange();
 	}
 }
 
diff --git a/Engine/Plugins/Media/MediaPlate/Source/MediaPlate/Public/MediaPlateComponent.h b/Engine/Plugins/Media/MediaPlate/Source/MediaPlate/Public/MediaPlateComponent.h
--- a/Engine/Plugins/Media/MediaPlate/Source/MediaPlate/Public/MediaPlateComponent.h
+++ b/Engine/Plugins/Media/MediaPlate/Source/MediaPlate/Public/MediaPlateComponent.h
@@ -113,6 +113,11 @@
 	 */
 	void UnregisterWithMediaTextureTracker();
 
+#if WITH_EDITOR
+	/** Propagate visible calculation changes to the tracker/player objects, optionally restarting playback if needed. */
+	void OnVisibleMipsTilesCalculationsChange();
+#endif
+
 private:
 	/** Name for our media component. */
 	static FLazyName MediaComponentName;
diff --git a/Engine/Plugins/Media/MediaPlate/Source/MediaPlateEditor/Private/MediaPlateCustomization.cpp b/Engine/Plugins/Media/MediaPlate/Source/MediaPlateEditor/Private/MediaPlateCustomization.cpp
--- a/Engine/Plugins/Media/MediaPlate/Source/MediaPlateEditor/Private/MediaPlateCustomization.cpp
+++ b/Engine/Plugins/Media/MediaPlate/Source/MediaPlateEditor/Private/MediaPlateCustomization.cpp
@@ -433,6 +433,7 @@
 			{
 				// Update the setting in the media plate.
 				MediaPlate->VisibleMipsTilesCalculations = MeshMode;
+				MediaPlate->OnVisibleMipsTilesCalculationsChange();
 
 				// Set the appropriate mesh.
 				if (MeshMode == EMediaTextureVisibleMipsTiles::Plane)
