MediaPlate: Added control of the horizontal arc size of sphere meshes.

#rnx

diff --git a/Engine/Plugins/Media/MediaPlate/Source/MediaPlate/Public/MediaPlateComponent.h b/Engine/Plugins/Media/MediaPlate/Source/MediaPlate/Public/MediaPlateComponent.h
--- a/Engine/Plugins/Media/MediaPlate/Source/MediaPlate/Public/MediaPlateComponent.h
+++ b/Engine/Plugins/Media/MediaPlate/Source/MediaPlate/Public/MediaPlateComponent.h
@@ -104,6 +104,13 @@ public:
 	UPROPERTY(EditAnywhere, AdvancedDisplay, Category = "MediaPlate")
 	EMediaTextureVisibleMipsTiles VisibleMipsTilesCalculations;
 
+	/**
+	 * Horizontal size of the mesh used for visible mips and tiles calculations.
+	 * For a sphere, this is the horizontal arc size in degrees.
+	 */
+	UPROPERTY(EditAnywhere, AdvancedDisplay, Category = "MediaPlate")
+	float MeshHorizontalRange = 360.0f;
+
 	/**
 	 * Adds our media texture to the media texture tracker.
 	 */
diff --git a/Engine/Plugins/Media/MediaPlate/Source/MediaPlateEditor/Private/MediaPlateCustomization.cpp b/Engine/Plugins/Media/MediaPlate/Source/MediaPlateEditor/Private/MediaPlateCustomization.cpp
--- a/Engine/Plugins/Media/MediaPlate/Source/MediaPlateEditor/Private/MediaPlateCustomization.cpp
+++ b/Engine/Plugins/Media/MediaPlate/Source/MediaPlateEditor/Private/MediaPlateCustomization.cpp
@@ -373,6 +373,7 @@ void FMediaPlateCustomization::AddMeshCustomization(IDetailCategoryBuilder& Medi
 
 	// Visibility attributes.
 	TAttribute<EVisibility> MeshPlaneVisibility(this, &FMediaPlateCustomization::ShouldShowMeshPlaneWidgets);
+	TAttribute<EVisibility> MeshSphereVisibility(this, &FMediaPlateCustomization::ShouldShowMeshSphereWidgets);
 
 	// Add aspect ratio.
 	DetailGroup.AddWidgetRow()
@@ -414,6 +415,25 @@ void FMediaPlateCustomization::AddMeshCustomization(IDetailCategoryBuilder& Medi
 						.OnValueChanged(this, &FMediaPlateCustomization::SetAspectRatio)
 				]
 		];
+
+	// Add sphere horizontal arc.
+	DetailGroup.AddWidgetRow()
+		.Visibility(MeshSphereVisibility)
+		.NameContent()
+		[
+			SNew(STextBlock)
+				.Text(LOCTEXT("HorizontalArc", "Horizontal Arc"))
+				.ToolTipText(LOCTEXT("HorizontalArc_ToolTip",
+				"Sets the horizontal arc size of the sphere in degrees.\nFor example 360 for a full circle, 180 for a half circle."))
+				.Font(IDetailLayoutBuilder::GetDetailFont())
+		]
+		.ValueContent()
+		[
+			SNew(SNumericEntryBox<float>)
+				.Value(this, &FMediaPlateCustomization::GetMeshHorizontalRange)
+				.OnValueChanged(this, &FMediaPlateCustomization::SetMeshHorizontalRange)
+		];
+	
 }
 
 EVisibility FMediaPlateCustomization::ShouldShowMeshPlaneWidgets() const
@@ -421,6 +441,11 @@ EVisibility FMediaPlateCustomization::ShouldShowMeshPlaneWidgets() const
 	return (MeshMode == EMediaTextureVisibleMipsTiles::Plane) ? EVisibility::Visible : EVisibility::Hidden;
 }
 
+EVisibility FMediaPlateCustomization::ShouldShowMeshSphereWidgets() const
+{
+	return (MeshMode == EMediaTextureVisibleMipsTiles::Sphere) ? EVisibility::Visible : EVisibility::Hidden;
+}
+
 void FMediaPlateCustomization::SetMeshMode(EMediaTextureVisibleMipsTiles InMode)
 {
 	if (MeshMode != InMode)
@@ -442,13 +467,18 @@ void FMediaPlateCustomization::SetMeshMode(EMediaTextureVisibleMipsTiles InMode)
 				}
 				else if (MeshMode == EMediaTextureVisibleMipsTiles::Sphere)
 				{
-					MeshCustomization.SetSphereMesh(MediaPlate);
+					SetSphereMesh(MediaPlate);
 				}
 			}
 		}
 	}
 }
 
+void FMediaPlateCustomization::SetSphereMesh(UMediaPlateComponent* MediaPlate)
+{
+	MeshCustomization.SetSphereMesh(MediaPlate);
+}
+
 TSharedRef<SWidget> FMediaPlateCustomization::OnGetAspectRatios()
 {
 	FMenuBuilder MenuBuilder(true, NULL);
@@ -541,6 +571,40 @@ TOptional<float> FMediaPlateCustomization::GetAspectRatio() const
 	return TOptional<float>();
 }
 
+void FMediaPlateCustomization::SetMeshHorizontalRange(float HorizontalRange)
+{
+	HorizontalRange = FMath::Clamp(HorizontalRange, 0.0f, 360.0f);
+
+	// Loop through all our objects.
+	for (const TWeakObjectPtr<UMediaPlateComponent>& MediaPlatePtr : MediaPlatesList)
+	{
+		UMediaPlateComponent* MediaPlate = MediaPlatePtr.Get();
+		if (MediaPlate != nullptr)
+		{
+			if (MediaPlate->MeshHorizontalRange != HorizontalRange)
+			{
+				MediaPlate->MeshHorizontalRange = HorizontalRange;
+				SetSphereMesh(MediaPlate);
+			}
+		}
+	}
+}
+
+TOptional<float> FMediaPlateCustomization::GetMeshHorizontalRange() const
+{
+	// Loop through our objects.
+	for (const TWeakObjectPtr<UMediaPlateComponent>& MediaPlatePtr : MediaPlatesList)
+	{
+		UMediaPlateComponent* MediaPlate = MediaPlatePtr.Get();
+		if (MediaPlate != nullptr)
+		{
+			return MediaPlate->MeshHorizontalRange;
+		}
+	}
+
+	return TOptional<float>();
+}
+
 FString FMediaPlateCustomization::GetMediaSourcePath() const
 {
 	FString Path;
diff --git a/Engine/Plugins/Media/MediaPlate/Source/MediaPlateEditor/Private/MediaPlateCustomization.h b/Engine/Plugins/Media/MediaPlate/Source/MediaPlateEditor/Private/MediaPlateCustomization.h
--- a/Engine/Plugins/Media/MediaPlate/Source/MediaPlateEditor/Private/MediaPlateCustomization.h
+++ b/Engine/Plugins/Media/MediaPlate/Source/MediaPlateEditor/Private/MediaPlateCustomization.h
@@ -61,11 +61,21 @@
 	EVisibility ShouldShowMeshPlaneWidgets() const;
 
 	/**
+	 * Controls visibility for widgets for sphere meshes.
+	 */
+	EVisibility ShouldShowMeshSphereWidgets() const;
+
+	/**
 	 * Call this to switch between planes, spheres, etc.
 	 */
 	void SetMeshMode(EMediaTextureVisibleMipsTiles InMode);
 
 	/**
+	 * Call this to apply a sphere mesh to an object.
+	 */
+	void SetSphereMesh(UMediaPlateComponent* MediaPlate);
+
+	/**
 	 * Returns menu options for all aspect ratio presets.
 	 */
 	TSharedRef<SWidget> OnGetAspectRatios();
@@ -81,6 +91,16 @@
 	TOptional<float> GetAspectRatio() const;
 
 	/**
+	 * Call this to set the horizontal range of the mesh.
+	 */
+	void SetMeshHorizontalRange(float HorizontalRange);
+
+	/**
+	 * Call this to get the horizontal range of the mesh.
+	 */
+	TOptional<float> GetMeshHorizontalRange() const;
+
+	/**
 	 * Gets the object path for the media source object.
 	 */
 	FString GetMediaSourcePath() const;
diff --git a/Engine/Plugins/Media/MediaPlate/Source/MediaPlateEditor/Private/MediaPlateCustomizationMesh.cpp b/Engine/Plugins/Media/MediaPlate/Source/MediaPlateEditor/Private/MediaPlateCustomizationMesh.cpp
--- a/Engine/Plugins/Media/MediaPlate/Source/MediaPlateEditor/Private/MediaPlateCustomizationMesh.cpp
+++ b/Engine/Plugins/Media/MediaPlate/Source/MediaPlateEditor/Private/MediaPlateCustomizationMesh.cpp
@@ -54,14 +54,14 @@ void FMediaPlateCustomizationMesh::SetSphereMesh(UMediaPlateComponent* MediaPlat
 	if (StaticMeshComponent != nullptr)
 	{
 		// Do we already have this mesh?
-		FString AssetPath = GetAssetPath();
+		FString AssetPath = GetAssetPath(MediaPlate);
 		UStaticMesh* StaticMesh = Cast<UStaticMesh>(StaticLoadObject(UStaticMesh::StaticClass(), NULL,
 			*AssetPath));
 		if (StaticMesh == nullptr)
 		{
 			// Create mesh.
 			FDynamicMesh3 NewMesh;
-			GenerateSphereMesh(&NewMesh);
+			GenerateSphereMesh(&NewMesh, MediaPlate);
 
 			// Create asset.
 			StaticMesh = CreateStaticMeshAsset(&NewMesh, AssetPath);
@@ -93,10 +93,12 @@ void FMediaPlateCustomizationMesh::SetMesh(UStaticMeshComponent* StaticMeshCompo
 	}
 }
 
-void FMediaPlateCustomizationMesh::GenerateSphereMesh(FDynamicMesh3* OutMesh)
+void FMediaPlateCustomizationMesh::GenerateSphereMesh(FDynamicMesh3* OutMesh,
+	UMediaPlateComponent* MediaPlate)
 {
 	FMediaPlateSphereGenerator SphereGen;
 	SphereGen.Radius = 50.0f;
+	SphereGen.ThetaRange = FMath::DegreesToRadians(MediaPlate->MeshHorizontalRange);
 	SphereGen.NumTheta = 17;
 	SphereGen.NumPhi = 17;
 	SphereGen.bPolygroupPerQuad = false;
@@ -137,10 +139,14 @@ UStaticMesh* FMediaPlateCustomizationMesh::CreateStaticMeshAsset(FDynamicMesh3*
 	return NewStaticMesh;
 }
 
-FString FMediaPlateCustomizationMesh::GetAssetPath()
+FString FMediaPlateCustomizationMesh::GetAssetPath(UMediaPlateComponent* MediaPlate)
 {
-	FString AssetPath = "/Game/_MediaPlate/Sphere";
-	
+	// Add the horizontal range.
+	FString ID = FString::SanitizeFloat(MediaPlate->MeshHorizontalRange);
+	ID.ReplaceCharInline(TCHAR('.'), TCHAR('_'));
+
+	FString AssetPath = FString::Printf(TEXT("/Game/_MediaPlate/Sphere_%s"), *ID);
+
 	return AssetPath;
 }
 
diff --git a/Engine/Plugins/Media/MediaPlate/Source/MediaPlateEditor/Private/MediaPlateCustomizationMesh.h b/Engine/Plugins/Media/MediaPlate/Source/MediaPlateEditor/Private/MediaPlateCustomizationMesh.h
--- a/Engine/Plugins/Media/MediaPlate/Source/MediaPlateEditor/Private/MediaPlateCustomizationMesh.h
+++ b/Engine/Plugins/Media/MediaPlate/Source/MediaPlateEditor/Private/MediaPlateCustomizationMesh.h
@@ -1,51 +1,51 @@
-// Copyright Epic Games, Inc. All Rights Reserved.
-
-#pragma once
-
-#include "GeometryBase.h"
-
-PREDECLARE_USE_GEOMETRY_CLASS(FDynamicMesh3);
-
-class UMediaPlateComponent;
-class UStaticMesh;
-class UStaticMeshComponent;
-
-/**
- * Implements mesh handling for FMediaPlateCustomization.
- */
-class FMediaPlateCustomizationMesh
-{
-public:
-	/**
-	 * Call this to use a plane mesh on a media plate.
-	 */
-	void SetPlaneMesh(UMediaPlateComponent* MediaPlate);
-
-	/**
-	 * Call this to use a sphere mesh on a media plate.
-	*/
-	void SetSphereMesh(UMediaPlateComponent* MediaPlate);
-
-private:
-
-	/**
-	 * Generic function to set a mesh on a StaticMeshComponent.
-	 */
-	void SetMesh(UStaticMeshComponent* StaticMeshComponent, UStaticMesh* Mesh);
-
-	/**
-	 * Creates a sphere mesh.
-	 */
-	void GenerateSphereMesh(FDynamicMesh3* OutMesh);
-
-	/**
-	 * Call this to create a static mesh asset for a mesh.
-	 */
-	UStaticMesh* CreateStaticMeshAsset(FDynamicMesh3* Mesh, const FString& AssetPath);
-
-	/**
-	 * Gets the path for the mesh asset.
-	 */
-	FString GetAssetPath();
-};
-
+// Copyright Epic Games, Inc. All Rights Reserved.
+
+#pragma once
+
+#include "GeometryBase.h"
+
+PREDECLARE_USE_GEOMETRY_CLASS(FDynamicMesh3);
+
+class UMediaPlateComponent;
+class UStaticMesh;
+class UStaticMeshComponent;
+
+/**
+ * Implements mesh handling for FMediaPlateCustomization.
+ */
+class FMediaPlateCustomizationMesh
+{
+public:
+	/**
+	 * Call this to use a plane mesh on a media plate.
+	 */
+	void SetPlaneMesh(UMediaPlateComponent* MediaPlate);
+
+	/**
+	 * Call this to use a sphere mesh on a media plate.
+	*/
+	void SetSphereMesh(UMediaPlateComponent* MediaPlate);
+
+private:
+
+	/**
+	 * Generic function to set a mesh on a StaticMeshComponent.
+	 */
+	void SetMesh(UStaticMeshComponent* StaticMeshComponent, UStaticMesh* Mesh);
+
+	/**
+	 * Creates a sphere mesh.
+	 */
+	void GenerateSphereMesh(FDynamicMesh3* OutMesh, UMediaPlateComponent* MediaPlate);
+
+	/**
+	 * Call this to create a static mesh asset for a mesh.
+	 */
+	UStaticMesh* CreateStaticMeshAsset(FDynamicMesh3* Mesh, const FString& AssetPath);
+
+	/**
+	 * Gets the path for the mesh asset.
+	 */
+	FString GetAssetPath(UMediaPlateComponent* MediaPlate);
+};
+
diff --git a/Engine/Plugins/Media/MediaPlate/Source/MediaPlateEditor/Private/MediaPlateSphereGenerator.h b/Engine/Plugins/Media/MediaPlate/Source/MediaPlateEditor/Private/MediaPlateSphereGenerator.h
--- a/Engine/Plugins/Media/MediaPlate/Source/MediaPlateEditor/Private/MediaPlateSphereGenerator.h
+++ b/Engine/Plugins/Media/MediaPlate/Source/MediaPlateEditor/Private/MediaPlateSphereGenerator.h
@@ -14,6 +14,8 @@ class FMediaPlateSphereGenerator : public UE::Geometry::FMeshShapeGenerator
 public:
 	/** Radius */
 	double Radius = 1;
+	/** Range for the horizontal arc. 360 degrees for a sphere, etc... */
+	double ThetaRange = FMathd::TwoPi;
 
 	int NumPhi = 3; // number of vertices along vertical extent from north pole to south pole
 	int NumTheta = 3; // number of vertices around circles
@@ -22,6 +24,9 @@ public:
 	bool bPolygroupPerQuad = false;
 
 private:
+	/** True if this is a closed mesh, e.g. a sphere. */
+	bool bIsClosed = false;
+
 	static FVector3d SphericalToCartesian(double r, double theta, double phi)
 	{
 		double Sphi = sin(phi);
@@ -42,8 +47,9 @@ private:
 			NormalParentVertex[VtxIdx] = VtxIdx;
 		};
 		{
+			int32 ThetaDivide = bIsClosed ? NumTheta : NumTheta -1;
 			const double Dphi = FMathd::Pi / double(NumPhi - 1);
-			const double Dtheta = FMathd::TwoPi / double(NumTheta);
+			const double Dtheta = ThetaRange / double(ThetaDivide);
 
 			int32 p,t;
 			double Phi, Theta;
@@ -68,8 +74,9 @@ private:
 	void GenerateUVVertices()
 	{
 		// generate the UV's
+		int32 ThetaDivide = bIsClosed ? NumTheta : NumTheta - 1;
 		const float DUVphi = 1.0 / float(NumPhi - 1);
-		const float DUVtheta = -1.0 / float(NumTheta);
+		const float DUVtheta = -1.0 / float(ThetaDivide);
 
 		int32 UVIdx = 0;
 		int32 p,t;
@@ -133,12 +140,16 @@ private:
 					PolyIdx++;
 				}
 			}
-			OutputTriangle(TriIdx++, PolyIdx,
-						   {Corners[0], Corners[1] - NumTheta, Corners[2] - NumTheta},
-						   {UVCorners[0]         , UVCorners[1],            UVCorners[2] });
-			OutputTriangle(TriIdx++, PolyIdx,
-						   {Corners[2] - NumTheta, Corners[3],   Corners[0]},
-						   {UVCorners[2],          UVCorners[3],            UVCorners[0]});
+			// Close up this mesh if required.
+			if (bIsClosed)
+			{
+				OutputTriangle(TriIdx++, PolyIdx,
+					{Corners[0], Corners[1] - NumTheta, Corners[2] - NumTheta},
+					{UVCorners[0]         , UVCorners[1],            UVCorners[2]});
+				OutputTriangle(TriIdx++, PolyIdx,
+					{Corners[2] - NumTheta, Corners[3],   Corners[0]},
+					{UVCorners[2],          UVCorners[3],            UVCorners[0]});
+			}
 			for (auto& i : Corners) ++i; 
 			for (auto& i : UVCorners) i += 2;
 			if (bPolygroupPerQuad)
@@ -155,6 +166,11 @@ private:
 		const int32 NorthPoleVtxIdx = NumEquatorialVtx;
 		const int32 SouthPoleVtxIdx = NumEquatorialVtx + 1;
 		int32 PolyIdx = NumTheta  * (NumPhi - 3);
+		// If the mesh is not closed then we have less tris.
+		if (bIsClosed == false)
+		{
+			PolyIdx -= (NumPhi - 3);
+		}
 		int32 TriIdx = PolyIdx * 2;
 		if (bPolygroupPerQuad == false)
 		{
@@ -162,7 +178,8 @@ private:
 		}
 
 		// Triangles that connect to north pole
-		for (int32 t = 0; t < NumTheta; ++t)
+		int32 NumTris = bIsClosed ? NumTheta : NumTheta - 1;
+		for (int32 t = 0; t < NumTris; ++t)
 		{
 			OutputTriangle(TriIdx++, PolyIdx,
 						   {t, NorthPoleVtxIdx,       (t + 1) % NumTheta},
@@ -176,7 +193,7 @@ private:
 		// Triangles that connect to South pole
 		const int32 Offset   = NumEquatorialVtx - NumTheta;
 		const int32 OffsetUV = NumEquatorialUVVtx - (NumTheta + 1);
-		for (int32 t = 0; t < NumTheta; ++t)
+		for (int32 t = 0; t < NumTris; ++t)
 		{
 			OutputTriangle(TriIdx++, PolyIdx,
 						   {t + Offset,   ((t + 1) % NumTheta) + Offset, SouthPoleVtxIdx},
@@ -191,12 +208,19 @@ public:
 	/** Generate the mesh */
 	FMediaPlateSphereGenerator& Generate() override
 	{
+		// Is this a closed sphere?
+		bIsClosed = ThetaRange >= FMath::DegreesToRadians(359.9f);
+		if (bIsClosed)
+		{
+			ThetaRange = FMathd::TwoPi;
+		}
+		
 		// enforce sane values for vertex counts
 		NumPhi = FMath::Max(NumPhi, 3);
 		NumTheta = FMath::Max(NumTheta, 3);
 		const int32 NumVertices = (NumPhi - 2) * NumTheta + 2;
 		const int32 NumUVs = (NumPhi - 2) * (NumTheta + 1) + (2 * NumTheta);
-		const int32 NumTris = (NumPhi - 2) * NumTheta * 2;
+		const int32 NumTris = (NumPhi - 2) * (bIsClosed ? NumTheta : (NumTheta - 1)) * 2;
 		SetBufferSizes(NumVertices, NumTris, NumUVs, NumVertices);
 
 		GenerateVertices();
