MediaPlate: Can drag and drop media files into the editor viewport to create a media plate.

#rnx

diff --git a/Engine/Plugins/Media/MediaPlate/Source/MediaPlateEditor/Private/ActorFactoryMediaPlate.cpp b/Engine/Plugins/Media/MediaPlate/Source/MediaPlateEditor/Private/ActorFactoryMediaPlate.cpp
--- a/Engine/Plugins/Media/MediaPlate/Source/MediaPlateEditor/Private/ActorFactoryMediaPlate.cpp
+++ b/Engine/Plugins/Media/MediaPlate/Source/MediaPlateEditor/Private/ActorFactoryMediaPlate.cpp
@@ -2,9 +2,11 @@
 
 #include "ActorFactoryMediaPlate.h"
 
+#include "Async/Async.h"
 #include "MediaSource.h"
 #include "MediaPlate.h"
 #include "MediaPlateComponent.h"
+#include "MediaPlateEditorModule.h"
 #include "MediaPlaylist.h"
 
 #define LOCTEXT_NAMESPACE "ActorFactoryMediaPlate"
@@ -65,6 +67,35 @@ void UActorFactoryMediaPlate::SetUpActor(UObject* Asset, AActor* Actor)
 				if (Playlist->Num() == 0)
 				{
 					Playlist->Add(MediaSource);
+
+					// Is this media source from a drag and drop?
+					FMediaPlateEditorModule* EditorModule = FModuleManager::LoadModulePtr<FMediaPlateEditorModule>("MediaPlateEditor");
+					if (EditorModule != nullptr)
+					{
+						bool bIsInDragDropCache = EditorModule->RemoveMediaSourceFromDragDropCache(MediaSource);
+						if (bIsInDragDropCache)
+						{
+							// Yes. Move this out of transient.
+							// Can't do it here as the asset is still being used.
+							TWeakObjectPtr<UMediaPlateComponent> MediaPlateComponentPtr(MediaPlate->MediaPlateComponent);
+							AsyncTask(ENamedThreads::GameThread, [MediaPlateComponentPtr]()
+							{
+								UMediaPlateComponent* MediaPlateComponent = MediaPlateComponentPtr.Get();
+								if (MediaPlateComponent != nullptr)
+								{
+									UMediaPlaylist* Playlist = MediaPlateComponent->MediaPlaylist;
+									if ((Playlist != nullptr) && (Playlist->Num() > 0))
+									{
+										UMediaSource* MediaSource = Playlist->Get(0);
+										if (MediaSource != nullptr)
+										{
+											MediaSource->Rename(nullptr, MediaPlateComponent);
+										}
+									}
+								}
+							});
+						}
+					}
 				}
 			}
 		}
diff --git a/Engine/Plugins/Media/MediaPlate/Source/MediaPlateEditor/Private/MediaPlateEditorModule.cpp b/Engine/Plugins/Media/MediaPlate/Source/MediaPlateEditor/Private/MediaPlateEditorModule.cpp
--- a/Engine/Plugins/Media/MediaPlate/Source/MediaPlateEditor/Private/MediaPlateEditorModule.cpp
+++ b/Engine/Plugins/Media/MediaPlate/Source/MediaPlateEditor/Private/MediaPlateEditorModule.cpp
@@ -11,9 +11,12 @@
 #include "MediaPlateEditorStyle.h"
 #include "MediaPlayer.h"
 #include "MediaSoundComponent.h"
+#include "MediaSource.h"
 #include "Models/MediaPlateEditorCommands.h"
 #include "PropertyEditorModule.h"
 #include "Sequencer/MediaPlateTrackEditor.h"
+#include "SLevelViewport.h"
+#include "Subsystems/EditorAssetSubsystem.h"
 
 DEFINE_LOG_CATEGORY(LogMediaPlateEditor);
 
@@ -34,10 +37,21 @@ void FMediaPlateEditorModule::StartupModule()
 
 	ISequencerModule& SequencerModule = FModuleManager::Get().LoadModuleChecked<ISequencerModule>("Sequencer");
 	TrackEditorBindingHandle = SequencerModule.RegisterPropertyTrackEditor<FMediaPlateTrackEditor>();
+
+	FCoreDelegates::OnPostEngineInit.AddRaw(this, &FMediaPlateEditorModule::OnPostEngineInit);
 }
 
 void FMediaPlateEditorModule::ShutdownModule()
 {
+	if (GEditor != nullptr)
+	{
+		UEditorAssetSubsystem* EditorAssetSubsystem = GEditor->GetEditorSubsystem<UEditorAssetSubsystem>();
+		if (EditorAssetSubsystem != nullptr)
+		{
+			EditorAssetSubsystem->GetOnExtractAssetFromFile().RemoveAll(this);
+		}
+	}
+
 	ISequencerModule* SequencerModulePtr = FModuleManager::Get().GetModulePtr<ISequencerModule>("Sequencer");
 	if (SequencerModulePtr)
 	{
@@ -112,6 +126,17 @@ void FMediaPlateEditorModule::MediaPlateStartedPlayback(TObjectPtr<UMediaPlateCo
 	}
 }
 
+bool FMediaPlateEditorModule::RemoveMediaSourceFromDragDropCache(UMediaSource* MediaSource)
+{
+	const FString* Key = MapFileToMediaSource.FindKey(MediaSource);
+	bool bIsInCache = Key != nullptr;
+	if (bIsInCache)
+	{
+		MapFileToMediaSource.Remove(*Key);
+	}
+	return bIsInCache;
+}
+
 void FMediaPlateEditorModule::RegisterAssetTools()
 {
 	IAssetTools& AssetTools = FModuleManager::LoadModuleChecked<FAssetToolsModule>("AssetTools").Get();
@@ -140,4 +165,45 @@ void FMediaPlateEditorModule::UnregisterAssetTools()
 	}
 }
 
+void FMediaPlateEditorModule::OnPostEngineInit()
+{
+	UEditorAssetSubsystem* EditorAssetSubsystem = GEditor->GetEditorSubsystem<UEditorAssetSubsystem>();
+	if (EditorAssetSubsystem != nullptr)
+	{
+		EditorAssetSubsystem->GetOnExtractAssetFromFile().AddRaw(this, &FMediaPlateEditorModule::ExtractAssetDataFromFiles);
+	}
+}
+
+void FMediaPlateEditorModule::ExtractAssetDataFromFiles(const TArray<FString>& Files,
+	TArray<FAssetData>& AssetDataArray)
+{
+	if (Files.Num() > 0)
+	{
+		// Do we have this already?
+		UMediaSource* MediaSource = nullptr;
+		TWeakObjectPtr<UMediaSource>* ExistingMediaSourcePtr = MapFileToMediaSource.Find(Files[0]);
+		if (ExistingMediaSourcePtr != nullptr)
+		{
+			MediaSource = ExistingMediaSourcePtr->Get();
+		}
+
+		// If we dont have it then create one now.
+		if (MediaSource == nullptr)
+		{
+			MediaSource = UMediaSource::SpawnMediaSourceForString(Files[0], GetTransientPackage());
+			if (MediaSource != nullptr)
+			{
+				MapFileToMediaSource.Emplace(Files[0], MediaSource);
+			}
+		}
+
+		// Return this via the array.
+		if (MediaSource != nullptr)
+		{
+			FAssetData AssetData(MediaSource);
+			AssetDataArray.Add(AssetData);
+		}
+	}
+}
+
 IMPLEMENT_MODULE(FMediaPlateEditorModule, MediaPlateEditor)
diff --git a/Engine/Plugins/Media/MediaPlate/Source/MediaPlateEditor/Public/MediaPlateEditorModule.h b/Engine/Plugins/Media/MediaPlate/Source/MediaPlateEditor/Public/MediaPlateEditorModule.h
--- a/Engine/Plugins/Media/MediaPlate/Source/MediaPlateEditor/Public/MediaPlateEditorModule.h
+++ b/Engine/Plugins/Media/MediaPlate/Source/MediaPlateEditor/Public/MediaPlateEditorModule.h
@@ -2,16 +2,20 @@
 
 #pragma once
 
+#include "Containers/Map.h"
 #include "CoreMinimal.h"
 #include "Logging/LogMacros.h"
 #include "Modules/ModuleManager.h"
 #include "TickableEditorObject.h"
 #include "UObject/ObjectPtr.h"
 
+class FDragDropEvent;
 class IAssetTools;
 class IAssetTypeActions;
 class ISlateStyle;
 class UMediaPlateComponent;
+class UMediaSource;
+struct FAssetData;
 
 /** Log category for this module. */
 DECLARE_LOG_CATEGORY_EXTERN(LogMediaPlateEditor, Log, All);
@@ -37,6 +41,16 @@ public:
 	 * Call this when a media plate starts playing so we can track it.
 	 */
 	void MediaPlateStartedPlayback(TObjectPtr<UMediaPlateComponent> MediaPlate);
+
+	/**
+	 * Call this if you want to check if your media source is a drag and drop media source.
+	 * If it is, then this will remove it from this module
+	 * and you can then change the outer (as it will be transient.)
+	 * 
+	 * @param MediaSource	Media source to check.
+	 * @return				True if it is a drag and drop media source.
+	 */
+	bool RemoveMediaSourceFromDragDropCache(UMediaSource* MediaSource);
 	
 private:
 
@@ -52,6 +66,9 @@ private:
 	/** Handle for our track editor. */
 	FDelegateHandle TrackEditorBindingHandle;
 
+	/** Maps a drag and drop file to a media source we created from it. */
+	TMap<FString, TWeakObjectPtr<UMediaSource>> MapFileToMediaSource;
+
 	/**
 	 * Registers all of our asset tools.
 	 */
@@ -69,4 +86,17 @@ private:
 	 * Unregisters all of our asset tools.
 	 */
 	void UnregisterAssetTools();
+
+	/**
+	 * Called after the engine has initialized.
+	 */
+	void OnPostEngineInit();
+
+	/**
+	 * Extracts an asset for a file if it can.
+	 *
+	 * @param Files				Files to get the asset from.
+	 * @param AssetDataArray	Will add the asset here if it gets one.
+	 */
+	void ExtractAssetDataFromFiles(const TArray<FString>& Files, TArray<FAssetData>& AssetDataArray);
 };
