MediaPlate: Fixed unused generated meshes lingering around.

#rnx

diff --git a/Engine/Plugins/Media/MediaPlate/Source/MediaPlateEditor/Private/MediaPlateCustomizationMesh.cpp b/Engine/Plugins/Media/MediaPlate/Source/MediaPlateEditor/Private/MediaPlateCustomizationMesh.cpp
--- a/Engine/Plugins/Media/MediaPlate/Source/MediaPlateEditor/Private/MediaPlateCustomizationMesh.cpp
+++ b/Engine/Plugins/Media/MediaPlate/Source/MediaPlateEditor/Private/MediaPlateCustomizationMesh.cpp
@@ -12,6 +12,8 @@
 
 #define LOCTEXT_NAMESPACE "FMediaPlateCustomizationMesh"
 
+TMap<UStaticMesh*, int32> FMediaPlateCustomizationMesh::MeshRefCount;
+
 void FMediaPlateCustomizationMesh::SetPlaneMesh(UMediaPlateComponent* MediaPlate)
 {
 	// Get static mesh component.
@@ -65,6 +67,27 @@ void FMediaPlateCustomizationMesh::SetSphereMesh(UMediaPlateComponent* MediaPlat
 
 			// Create asset.
 			StaticMesh = CreateStaticMeshAsset(&NewMesh, AssetPath);
+			if (StaticMesh != nullptr)
+			{
+				MeshRefCount.Add(StaticMesh, 1);
+			}
+		}
+		else
+		{
+			// Is this one of our generated meshes?
+			int32* CountPtr = MeshRefCount.Find(StaticMesh);
+			if (CountPtr != nullptr)
+			{
+				// Make sure its not transient.
+				UPackage* MeshPackage = StaticMesh->GetPackage();
+				if (MeshPackage != nullptr)
+				{
+					MeshPackage->ClearFlags(RF_Transient);
+				}
+
+				// Increase reference count.
+				(*CountPtr)++;
+			}
 		}
 
 		// Apply mesh.
@@ -76,6 +99,24 @@ void FMediaPlateCustomizationMesh::SetMesh(UStaticMeshComponent* StaticMeshCompo
 {
 	if (StaticMeshComponent != nullptr)
 	{
+		// Get existing mesh.
+		TObjectPtr<UStaticMesh> OldMesh = StaticMeshComponent->GetStaticMesh();
+		if (OldMesh != nullptr)
+		{
+			// Is this one of our generated meshes?
+			int32* CountPtr = MeshRefCount.Find(OldMesh);
+			if (CountPtr != nullptr)
+			{
+				// Update reference count.
+				(*CountPtr)--;
+				if (*CountPtr <= 0)
+				{
+					// No longer needed, so mark as transient.
+					OldMesh->GetPackage()->SetFlags(RF_Transient);
+				}
+			}
+		}
+
 		// Apply mesh to component.
 		StaticMeshComponent->SetStaticMesh(Mesh);
 		StaticMeshComponent->SetRelativeScale3D(FVector::OneVector);
diff --git a/Engine/Plugins/Media/MediaPlate/Source/MediaPlateEditor/Private/MediaPlateCustomizationMesh.h b/Engine/Plugins/Media/MediaPlate/Source/MediaPlateEditor/Private/MediaPlateCustomizationMesh.h
--- a/Engine/Plugins/Media/MediaPlate/Source/MediaPlateEditor/Private/MediaPlateCustomizationMesh.h
+++ b/Engine/Plugins/Media/MediaPlate/Source/MediaPlateEditor/Private/MediaPlateCustomizationMesh.h
@@ -47,5 +47,8 @@ private:
 	 * Gets the path for the mesh asset.
 	 */
 	FString GetAssetPath(UMediaPlateComponent* MediaPlate);
+
+	/** Reference count for our generated meshes. */
+	static TMap<UStaticMesh*, int32> MeshRefCount;
 };
 
