MediaPlate: Added menu option to Sequencer track editor to import from the media plate.
#rnx

diff --git a/Engine/Plugins/Media/MediaPlate/Source/MediaPlateEditor/Private/Sequencer/MediaPlateTrackEditor.cpp b/Engine/Plugins/Media/MediaPlate/Source/MediaPlateEditor/Private/Sequencer/MediaPlateTrackEditor.cpp
--- a/Engine/Plugins/Media/MediaPlate/Source/MediaPlateEditor/Private/Sequencer/MediaPlateTrackEditor.cpp
+++ b/Engine/Plugins/Media/MediaPlate/Source/MediaPlateEditor/Private/Sequencer/MediaPlateTrackEditor.cpp
@@ -33,6 +33,23 @@ FMediaPlateTrackEditor::~FMediaPlateTrackEditor()
 {
 }
 
+void FMediaPlateTrackEditor::BuildObjectBindingContextMenu(FMenuBuilder& MenuBuilder, const TArray<FGuid>& ObjectBindings, const UClass* ObjectClass)
+{
+	if ((ObjectClass != nullptr) && (ObjectClass->IsChildOf(AMediaPlate::StaticClass())))
+	{
+		MenuBuilder.BeginSection("Media Plate", LOCTEXT("MediaPlate", "Media Plate"));
+
+		MenuBuilder.AddMenuEntry(
+			LOCTEXT("Import", "Import"),
+			LOCTEXT("ImportTooltip", "Import tracks from the Media Plate object."),
+			FSlateIcon(),
+			FUIAction(FExecuteAction::CreateRaw(this,
+				&FMediaPlateTrackEditor::ImportObjectBinding, ObjectBindings)));
+	
+		MenuBuilder.EndSection();
+	}
+}
+
 void FMediaPlateTrackEditor::BuildObjectBindingTrackMenu(FMenuBuilder& MenuBuilder, const TArray<FGuid>& ObjectBindings, const UClass* ObjectClass)
 {
 	// Is this a media plate?
@@ -144,4 +161,25 @@ void FMediaPlateTrackEditor::OnRelease()
 	FMovieSceneTrackEditor::OnRelease();
 }
 
+void FMediaPlateTrackEditor::ImportObjectBinding(const TArray<FGuid> ObjectBindings)
+{
+	// Get actor.
+	const TSharedPtr<ISequencer> SequencerPtr = GetSequencer();
+	for (const FGuid ObjectBinding : ObjectBindings)
+	{
+		UObject* BoundObject = SequencerPtr.IsValid() ? SequencerPtr->FindSpawnedObjectOrTemplate(ObjectBinding) : nullptr;
+		AActor* Actor = Cast<AActor>(BoundObject);
+		if (Actor != nullptr)
+		{
+			// Get media plate componnent.
+			if (UMediaPlateComponent* MediaPlateComponent = Actor->FindComponentByClass<UMediaPlateComponent>())
+			{
+				// Add tracks for this.
+				AddTrackForComponent(MediaPlateComponent);
+			}
+		}
+	}
+}
+
+
 #undef LOCTEXT_NAMESPACE
diff --git a/Engine/Plugins/Media/MediaPlate/Source/MediaPlateEditor/Private/Sequencer/MediaPlateTrackEditor.h b/Engine/Plugins/Media/MediaPlate/Source/MediaPlateEditor/Private/Sequencer/MediaPlateTrackEditor.h
--- a/Engine/Plugins/Media/MediaPlate/Source/MediaPlateEditor/Private/Sequencer/MediaPlateTrackEditor.h
+++ b/Engine/Plugins/Media/MediaPlate/Source/MediaPlateEditor/Private/Sequencer/MediaPlateTrackEditor.h
@@ -51,6 +51,7 @@ public:
 public:
 
 	//~ FMovieSceneTrackEditor interface
+	virtual void BuildObjectBindingContextMenu(FMenuBuilder& MenuBuilder, const TArray<FGuid>& ObjectBindings, const UClass* ObjectClass) override;
 	virtual void BuildObjectBindingTrackMenu(FMenuBuilder& MenuBuilder, const TArray<FGuid>& ObjectBindings, const UClass* ObjectClass) override;
 	virtual bool SupportsType(TSubclassOf<UMovieSceneTrack> TrackClass) const override;
 	//~ ISequencerTrackEditor interface
@@ -70,6 +71,11 @@ private:
 	 */
 	void AddTrackForComponent(UMediaPlateComponent* Component);
 
+	/**
+	 * Adds tracks from an object binding.
+	 */
+	void ImportObjectBinding(const TArray<FGuid> ObjectBindings);
+
 	/** Handle to our delegate. */
 	FDelegateHandle OnActorAddedToSequencerHandle;
 };
