MediaPlate: Increasing generated mesh resolution to 64, before those are converted to user-exposed parameters.

diff --git a/Engine/Plugins/Media/MediaPlate/Source/MediaPlateEditor/Private/MediaPlateCustomizationMesh.cpp b/Engine/Plugins/Media/MediaPlate/Source/MediaPlateEditor/Private/MediaPlateCustomizationMesh.cpp
--- a/Engine/Plugins/Media/MediaPlate/Source/MediaPlateEditor/Private/MediaPlateCustomizationMesh.cpp
+++ b/Engine/Plugins/Media/MediaPlate/Source/MediaPlateEditor/Private/MediaPlateCustomizationMesh.cpp
@@ -140,8 +140,8 @@ void FMediaPlateCustomizationMesh::GenerateSphereMesh(FDynamicMesh3* OutMesh,
 	FMediaPlateSphereGenerator SphereGen;
 	SphereGen.Radius = 50.0f;
 	SphereGen.ThetaRange = FMath::DegreesToRadians(MediaPlate->MeshHorizontalRange);
-	SphereGen.NumTheta = 17;
-	SphereGen.NumPhi = 17;
+	SphereGen.NumTheta = 64;
+	SphereGen.NumPhi = 64;
 	SphereGen.bPolygroupPerQuad = false;
 	SphereGen.Generate();
 
