MediaPlate: Cache settings get copied over to the sequencer track when importing from a media plate.
#rnx

diff --git a/Engine/Plugins/Media/MediaPlate/Source/MediaPlateEditor/Private/Sequencer/MediaPlateTrackEditor.cpp b/Engine/Plugins/Media/MediaPlate/Source/MediaPlateEditor/Private/Sequencer/MediaPlateTrackEditor.cpp
--- a/Engine/Plugins/Media/MediaPlate/Source/MediaPlateEditor/Private/Sequencer/MediaPlateTrackEditor.cpp
+++ b/Engine/Plugins/Media/MediaPlate/Source/MediaPlateEditor/Private/Sequencer/MediaPlateTrackEditor.cpp
@@ -8,6 +8,7 @@
 #include "MediaPlate.h"
 #include "MediaPlateComponent.h"
 #include "MediaPlaylist.h"
+#include "MovieSceneMediaSection.h"
 #include "MovieSceneMediaTrack.h"
 
 #define LOCTEXT_NAMESPACE "FMediaPlateTrackEditor"
@@ -141,7 +142,14 @@ void FMediaPlateTrackEditor::AddTrackForComponent(UMediaPlateComponent* Componen
 				UMediaSource* MediaSource = Playlist->Get(Index);
 				if (MediaSource != nullptr)
 				{
-					MediaTrack->AddNewMediaSource(*MediaSource, FFrameNumber(0));
+					UMovieSceneSection* Section = MediaTrack->AddNewMediaSource(*MediaSource, FFrameNumber(0));
+					
+					// Copy cache settings from media plate.
+					UMovieSceneMediaSection* MediaSection = Cast<UMovieSceneMediaSection>(Section);
+					if (MediaSection != nullptr)
+					{
+						MediaSection->CacheSettings = Component->CacheSettings;
+					}
 				}
 			}
 		}
