MediaPlate: Added auto aspect ratio.

#rnx

diff --git a/Engine/Plugins/Media/MediaPlate/Source/MediaPlate/Private/MediaPlate.cpp b/Engine/Plugins/Media/MediaPlate/Source/MediaPlate/Private/MediaPlate.cpp
--- a/Engine/Plugins/Media/MediaPlate/Source/MediaPlate/Private/MediaPlate.cpp
+++ b/Engine/Plugins/Media/MediaPlate/Source/MediaPlate/Private/MediaPlate.cpp
@@ -26,6 +26,7 @@ AMediaPlate::AMediaPlate(const FObjectInitializer& ObjectInitializer)
 	// Set up static mesh component.
 	StaticMeshComponent = CreateDefaultSubobject<UStaticMeshComponent>(TEXT("StaticMeshComponent"));
 	StaticMeshComponent->SetupAttachment(RootComponent);
+	MediaPlateComponent->StaticMeshComponent = StaticMeshComponent;
 
 	// Hook up mesh.
 	struct FConstructorStatics
diff --git a/Engine/Plugins/Media/MediaPlate/Source/MediaPlate/Private/MediaPlateComponent.cpp b/Engine/Plugins/Media/MediaPlate/Source/MediaPlate/Private/MediaPlateComponent.cpp
--- a/Engine/Plugins/Media/MediaPlate/Source/MediaPlate/Private/MediaPlateComponent.cpp
+++ b/Engine/Plugins/Media/MediaPlate/Source/MediaPlate/Private/MediaPlateComponent.cpp
@@ -4,6 +4,9 @@
 
 #include "Components/StaticMeshComponent.h"
 #include "Engine/StaticMesh.h"
+#include "IMediaClock.h"
+#include "IMediaClockSink.h"
+#include "IMediaModule.h"
 #include "MediaComponent.h"
 #include "MediaPlateModule.h"
 #include "MediaPlayer.h"
@@ -15,6 +18,43 @@
 
 #define LOCTEXT_NAMESPACE "MediaPlate"
 
+
+/**
+ * Media clock sink for media textures.
+ */
+class FMediaComponentClockSink
+	: public IMediaClockSink
+{
+public:
+
+	FMediaComponentClockSink(UMediaPlateComponent* InOwner)
+		: Owner(InOwner)
+	{ }
+
+	virtual ~FMediaComponentClockSink() { }
+
+	virtual void TickOutput(FTimespan DeltaTime, FTimespan Timecode) override
+	{
+		if (UMediaPlateComponent* OwnerPtr = Owner.Get())
+		{
+			Owner->TickOutput();
+		}
+	}
+
+
+	/**
+	 * Call this when the owner is destroyed.
+	 */
+	void OwnerDestroyed()
+	{
+		Owner.Reset();
+	}
+
+private:
+
+	TWeakObjectPtr<UMediaPlateComponent> Owner;
+};
+
 FLazyName UMediaPlateComponent::MediaComponentName(TEXT("MediaComponent0"));
 FLazyName UMediaPlateComponent::MediaPlaylistName(TEXT("MediaPlaylist0"));
 
@@ -68,6 +108,25 @@ void UMediaPlateComponent::BeginPlay()
 	}
 }
 
+void UMediaPlateComponent::BeginDestroy()
+{
+	if (ClockSink.IsValid())
+	{
+		// Tell sink we are done.
+		ClockSink->OwnerDestroyed();
+
+		IMediaModule* MediaModule = FModuleManager::GetModulePtr<IMediaModule>("Media");
+		if (MediaModule != nullptr)
+		{
+			MediaModule->GetClock().RemoveSink(ClockSink.ToSharedRef());
+		}
+
+		ClockSink.Reset();
+	}
+
+	Super::BeginDestroy();
+}
+
 UMediaPlayer* UMediaPlateComponent::GetMediaPlayer()
 {
 	TObjectPtr<UMediaPlayer> MediaPlayer = nullptr;
@@ -109,6 +168,25 @@ void UMediaPlateComponent::Play()
 	{
 		UE_LOG(LogMediaPlate, Warning, TEXT("Could not play anything."));
 	}
+	else
+	{
+		// Are we using automatic aspect ratio?
+		if ((bIsAspectRatioAuto) &&
+			(VisibleMipsTilesCalculations == EMediaTextureVisibleMipsTiles::Plane))
+		{
+			// Start the clock sink so we can tick.
+			IMediaModule* MediaModule = FModuleManager::LoadModulePtr<IMediaModule>("Media");
+			if (MediaModule != nullptr)
+			{
+				if (ClockSink.IsValid() == false)
+				{
+					ClockSink = MakeShared<FMediaComponentClockSink, ESPMode::ThreadSafe>(this);
+				}
+				bIsWaitingForRender = true;
+				MediaModule->GetClock().AddSink(ClockSink.ToSharedRef());
+			}
+		}
+	}
 }
 
 void UMediaPlateComponent::Stop()
@@ -118,6 +196,8 @@ void UMediaPlateComponent::Stop()
 	{
 		MediaPlayer->Close();
 	}
+
+	StopClockSink();
 }
 
 void UMediaPlateComponent::SetMeshHorizontalRange(float InMeshHorizontalRange)
@@ -185,6 +265,72 @@ bool UMediaPlateComponent::PlayMediaSource(UMediaSource* InMediaSource)
 	return bIsPlaying;
 }
 
+void UMediaPlateComponent::SetAspectRatio(float AspectRatio)
+{
+	// Get the static mesh.
+	if (StaticMeshComponent != nullptr)
+	{
+		// Update the scale.
+		float Height = 1.0f;
+		if (AspectRatio != 0.0f)
+		{
+			Height = 1.0f / AspectRatio;
+		}
+		FVector Scale(1.0f, 1.0f, Height);
+#if WITH_EDITOR
+		StaticMeshComponent->Modify();
+#endif // WITH_EDITOR
+		StaticMeshComponent->SetRelativeScale3D(Scale);
+	}
+}
+
+void UMediaPlateComponent::TickOutput()
+{
+	TObjectPtr<UMediaPlayer> MediaPlayer = GetMediaPlayer();
+	if (MediaPlayer != nullptr)
+	{
+		// Is the player ready?
+		if (MediaPlayer->IsPreparing() == false)
+		{
+			if (bIsWaitingForRender)
+			{
+				bIsWaitingForRender = false;
+			}
+			else
+			{
+				// Get the texture.
+				UMediaTexture* Texture = GetMediaTexture();
+				if (Texture != nullptr)
+				{
+					float Width = Texture->GetSurfaceWidth();
+					float Height = Texture->GetSurfaceHeight();
+					// Set aspect ratio.
+					if (Height != 0.0f)
+					{
+						float AspectRatio = Width / Height;
+						SetAspectRatio(AspectRatio);
+					}
+
+					// No need to tick anymore.
+					StopClockSink();
+				}
+			}
+		}
+	}
+}
+
+void UMediaPlateComponent::StopClockSink()
+{
+	if (ClockSink.IsValid())
+	{
+		IMediaModule* MediaModule = FModuleManager::GetModulePtr<IMediaModule>("Media");
+		if (MediaModule != nullptr)
+		{
+			MediaModule->GetClock().RemoveSink(ClockSink.ToSharedRef());
+		}
+	}
+}
+
 #if WITH_EDITOR
 
 void UMediaPlateComponent::OnVisibleMipsTilesCalculationsChange()
diff --git a/Engine/Plugins/Media/MediaPlate/Source/MediaPlate/Public/MediaPlateComponent.h b/Engine/Plugins/Media/MediaPlate/Source/MediaPlate/Public/MediaPlateComponent.h
--- a/Engine/Plugins/Media/MediaPlate/Source/MediaPlate/Public/MediaPlateComponent.h
+++ b/Engine/Plugins/Media/MediaPlate/Source/MediaPlate/Public/MediaPlateComponent.h
@@ -10,6 +10,7 @@
 
 #include "MediaPlateComponent.generated.h"
 
+class FMediaComponentClockSink;
 class UMediaComponent;
 class UMediaPlayer;
 class UMediaPlaylist;
@@ -30,6 +31,7 @@ public:
 	//~ UActorComponent interface.
 	virtual void OnRegister();
 	virtual void BeginPlay();
+	virtual void BeginDestroy();
 #if WITH_EDITOR
 	virtual void PostEditChangeProperty(FPropertyChangedEvent& PropertyChangedEvent) override;
 #endif // WITH_EDITOR
@@ -74,6 +76,10 @@ public:
 	UPROPERTY(EditAnywhere, Category = "MediaPlate", meta = (ClampMin = "0.0"))
 	float StartTime = 0.0f;
 
+	/** If true then set the aspect ratio automatically based on the media. */
+	UPROPERTY(BlueprintReadWrite, Category = "MediaPlate")
+	bool bIsAspectRatioAuto = true;
+
 	/** Holds the media player. */
 	UPROPERTY(Category = MediaPlate, VisibleAnywhere, BlueprintReadOnly)
 	TObjectPtr<UMediaComponent> MediaComponent;
@@ -82,6 +88,10 @@ public:
 	UPROPERTY(EditAnywhere, AdvancedDisplay, Category = MediaPlate)
 	TObjectPtr<UMediaSoundComponent> SoundComponent;
 
+	/** Holds the component for the mesh. */
+	UPROPERTY(EditAnywhere, AdvancedDisplay, Category = MediaPlate)
+	TObjectPtr<UStaticMeshComponent> StaticMeshComponent;
+
 	/** What media playlist to play. */
 	UPROPERTY(EditAnywhere, Category = "MediaPlate")
 	TObjectPtr<UMediaPlaylist> MediaPlaylist;
@@ -114,6 +124,16 @@ public:
 	 */
 	void UnregisterWithMediaTextureTracker();
 
+	/**
+	 * Call this to set the aspect ratio of the mesh.
+	 */
+	void SetAspectRatio(float AspectRatio);
+
+	/**
+	 * Called from the media clock.
+	 */
+	void TickOutput();
+
 #if WITH_EDITOR
 	/** Propagate visible calculation changes to the tracker/player objects, optionally restarting playback if needed. */
 	void OnVisibleMipsTilesCalculationsChange();
@@ -130,6 +150,10 @@ private:
 
 	/** Info representing this object. */
 	TSharedPtr<FMediaTextureTrackerObject, ESPMode::ThreadSafe> MediaTextureTrackerObject;
+	/** Our media clock sink. */
+	TSharedPtr<FMediaComponentClockSink, ESPMode::ThreadSafe> ClockSink;
+	/** True if we are waiting for the first frame to come out. */
+	bool bIsWaitingForRender = false;
 
 	/**
 	 * Plays a media source.
@@ -138,4 +162,10 @@ private:
 	 * @return	True if we played anything.
 	 */
 	bool PlayMediaSource(UMediaSource* InMediaSource);
+
+	/**
+	 * Stops the clock sink so we no longer tick.
+	 */
+	void StopClockSink();
+
 };
diff --git a/Engine/Plugins/Media/MediaPlate/Source/MediaPlateEditor/Private/MediaPlateCustomization.cpp b/Engine/Plugins/Media/MediaPlate/Source/MediaPlateEditor/Private/MediaPlateCustomization.cpp
--- a/Engine/Plugins/Media/MediaPlate/Source/MediaPlateEditor/Private/MediaPlateCustomization.cpp
+++ b/Engine/Plugins/Media/MediaPlate/Source/MediaPlateEditor/Private/MediaPlateCustomization.cpp
@@ -18,6 +18,7 @@
 #include "Styling/AppStyle.h"
 #include "Subsystems/AssetEditorSubsystem.h"
 #include "Widgets/Input/SButton.h"
+#include "Widgets/Input/SCheckBox.h"
 #include "Widgets/Input/SComboButton.h"
 #include "Widgets/Input/SFilePathPicker.h"
 #include "Widgets/Input/SNumericEntryBox.h"
@@ -416,6 +417,24 @@ void FMediaPlateCustomization::AddMeshCustomization(IDetailCategoryBuilder& Medi
 				]
 		];
 
+	// Add auto aspect ratio.
+	DetailGroup.AddWidgetRow()
+		.Visibility(MeshPlaneVisibility)
+		.NameContent()
+		[
+			SNew(STextBlock)
+				.Text(LOCTEXT("AutoAspectRatio", "Auto Aspect Ratio"))
+				.ToolTipText(LOCTEXT("AutoAspectRatio_ToolTip",
+					"Sets the aspect ratio to match the media."))
+				.Font(IDetailLayoutBuilder::GetDetailFont())
+		]
+		.ValueContent()
+		[
+			SNew(SCheckBox)
+				.IsChecked(this, &FMediaPlateCustomization::IsAspectRatioAuto)
+				.OnCheckStateChanged(this, &FMediaPlateCustomization::SetIsAspectRatioAuto)
+		];
+
 	// Add sphere horizontal arc.
 	DetailGroup.AddWidgetRow()
 		.Visibility(MeshSphereVisibility)
@@ -479,6 +498,48 @@ void FMediaPlateCustomization::SetSphereMesh(UMediaPlateComponent* MediaPlate)
 	MeshCustomization.SetSphereMesh(MediaPlate);
 }
 
+ECheckBoxState FMediaPlateCustomization::IsAspectRatioAuto() const
+{
+	ECheckBoxState State = ECheckBoxState::Undetermined;
+
+	for (const TWeakObjectPtr<UMediaPlateComponent>& MediaPlatePtr : MediaPlatesList)
+	{
+		UMediaPlateComponent* MediaPlate = MediaPlatePtr.Get();
+		if (MediaPlate != nullptr)
+		{
+			ECheckBoxState NewState = MediaPlate->bIsAspectRatioAuto ? ECheckBoxState::Checked :
+				ECheckBoxState::Unchecked;
+			if (State == ECheckBoxState::Undetermined)
+			{
+				State = NewState;
+			}
+			else if (State != NewState)
+			{
+				// If the media plates have different states then return undetermined.
+				State = ECheckBoxState::Undetermined;
+				break;
+			}
+		}
+	}
+
+	return State;
+}
+
+void FMediaPlateCustomization::SetIsAspectRatioAuto(ECheckBoxState State)
+{
+	bool bEnable = (State == ECheckBoxState::Checked);
+
+	// Loop through all our objects.
+	for (const TWeakObjectPtr<UMediaPlateComponent>& MediaPlatePtr : MediaPlatesList)
+	{
+		UMediaPlateComponent* MediaPlate = MediaPlatePtr.Get();
+		if (MediaPlate != nullptr)
+		{
+			MediaPlate->bIsAspectRatioAuto = bEnable;
+		}
+	}
+}
+
 TSharedRef<SWidget> FMediaPlateCustomization::OnGetAspectRatios()
 {
 	FMenuBuilder MenuBuilder(true, NULL);
@@ -500,33 +561,13 @@ TSharedRef<SWidget> FMediaPlateCustomization::OnGetAspectRatios()
 
 void FMediaPlateCustomization::SetAspectRatio(float AspectRatio)
 {
-	// Calculate required scale from the aspect ratio.
-	float Height = 1.0f;
-	if (AspectRatio != 0.0f)
-	{
-		Height = 1.0f / AspectRatio;
-	}
-	FVector Scale(1.0f, 1.0f, Height);
-
 	// Loop through all our objects.
 	for (const TWeakObjectPtr<UMediaPlateComponent>& MediaPlatePtr : MediaPlatesList)
 	{
 		UMediaPlateComponent* MediaPlate = MediaPlatePtr.Get();
 		if (MediaPlate != nullptr)
 		{
-			AActor* Owner = MediaPlate->GetOwner();
-			AMediaPlate* MediaPlateActor = Cast<AMediaPlate>(Owner);
-			if (MediaPlateActor != nullptr)
-			{
-				// Get the static mesh.
-				UStaticMeshComponent* StaticMeshComponent = MediaPlateActor->StaticMeshComponent;
-				if (StaticMeshComponent != nullptr)
-				{
-					// Update the scale.
-					StaticMeshComponent->SetRelativeScale3D(Scale);
-					StaticMeshComponent->MarkPackageDirty();
-				}
-			}
+			MediaPlate->SetAspectRatio(AspectRatio);
 		}
 	}
 
diff --git a/Engine/Plugins/Media/MediaPlate/Source/MediaPlateEditor/Private/MediaPlateCustomization.h b/Engine/Plugins/Media/MediaPlate/Source/MediaPlateEditor/Private/MediaPlateCustomization.h
--- a/Engine/Plugins/Media/MediaPlate/Source/MediaPlateEditor/Private/MediaPlateCustomization.h
+++ b/Engine/Plugins/Media/MediaPlate/Source/MediaPlateEditor/Private/MediaPlateCustomization.h
@@ -8,6 +8,7 @@
 #include "Input/Reply.h"
 #include "MediaPlateCustomizationMesh.h"
 #include "MediaTextureTracker.h"
+#include "Styling/SlateTypes.h"
 
 class IDetailCategoryBuilder;
 class IDetailLayoutBuilder;
@@ -80,6 +81,16 @@ private:
 	 */
 	TSharedRef<SWidget> OnGetAspectRatios();
 
+	/**
+	 * Call this to see if auto aspect ratio is enabled.
+	 */
+	ECheckBoxState IsAspectRatioAuto() const;
+
+	/**
+	 * Call this to enable/disable automatic aspect ratio.
+	 */
+	void SetIsAspectRatioAuto(ECheckBoxState State);
+
 	/**
 	 * Call this to set the aspect ratio.
 	 */
