MediaPlate: Added vertical arc control of sphere mesh.

#rnx

diff --git a/Engine/Plugins/Media/ImgMedia/Source/ImgMedia/Private/Assets/ImgMediaMipMapInfo.cpp b/Engine/Plugins/Media/ImgMedia/Source/ImgMedia/Private/Assets/ImgMediaMipMapInfo.cpp
--- a/Engine/Plugins/Media/ImgMedia/Source/ImgMedia/Private/Assets/ImgMediaMipMapInfo.cpp
+++ b/Engine/Plugins/Media/ImgMedia/Source/ImgMedia/Private/Assets/ImgMediaMipMapInfo.cpp
@@ -547,9 +547,9 @@ namespace {
 	class FSphereObjectInfo : public FImgMediaMipMapObjectInfo
 	{
 	public:
-		FSphereObjectInfo(UMeshComponent* InMeshComponent, float InLODBias, float InMeshHorizontalRange)
+		FSphereObjectInfo(UMeshComponent* InMeshComponent, float InLODBias, FVector2D InMeshRange)
 			: FImgMediaMipMapObjectInfo(InMeshComponent, InLODBias)
-			, MeshHorizontalRange(InMeshHorizontalRange)
+			, MeshRange(InMeshRange)
 		{
 		}
 
@@ -578,7 +578,9 @@ namespace {
 			auto TransformSphericalUVsToLocationWS = [this, &MeshTransform, DefaultSphereRadius](const FVector2D& InUV) -> FVector
 			{
 				// Convert from latlong UV to spherical coordinates
-				FVector2D TileCornerSpherical = FVector2D(UE_PI * InUV.Y, FMath::DegreesToRadians(MeshHorizontalRange) * (1.0f - InUV.X));
+				FVector2D TileCornerSpherical = FVector2D(
+					FMath::DegreesToRadians(MeshRange.Y) * InUV.Y,
+					FMath::DegreesToRadians(MeshRange.X) * (1.0f - InUV.X));
 				FVector CornersWS = TileCornerSpherical.SphericalToUnitCartesian() * DefaultSphereRadius;
 				return MeshTransform.TransformPosition(CornersWS);
 			};
@@ -654,7 +656,7 @@ namespace {
 		}
 
 	private:
-		float MeshHorizontalRange;
+		FVector2D MeshRange;
 	};
 
 } //end anonymous namespace
@@ -669,7 +671,8 @@ FImgMediaMipMapInfo::~FImgMediaMipMapInfo()
 	ClearAllObjects();
 }
 
-void FImgMediaMipMapInfo::AddObject(AActor* InActor, float LODBias, EMediaTextureVisibleMipsTiles MeshType, float MeshHorizontalRange)
+void FImgMediaMipMapInfo::AddObject(AActor* InActor, float LODBias, EMediaTextureVisibleMipsTiles MeshType,
+	FVector2D MeshRange)
 {
 	if (InActor != nullptr)
 	{
@@ -682,7 +685,7 @@ void FImgMediaMipMapInfo::AddObject(AActor* InActor, float LODBias, EMediaTextur
 				Objects.Add(new FPlaneObjectInfo(MeshComponent, LODBias));
 				break;
 			case EMediaTextureVisibleMipsTiles::Sphere:
-				Objects.Add(new FSphereObjectInfo(MeshComponent, LODBias, MeshHorizontalRange));
+				Objects.Add(new FSphereObjectInfo(MeshComponent, LODBias, MeshRange));
 				break;
 			default:
 				Objects.Add(new FImgMediaMipMapObjectInfo(MeshComponent, LODBias));
@@ -730,7 +733,7 @@ void FImgMediaMipMapInfo::AddObjectsUsingThisMediaTexture(UMediaTexture* InMedia
 				AActor* Owner = ObjectInfo->Object.Get();
 				if (Owner != nullptr)
 				{
-					AddObject(Owner,ObjectInfo->MipMapLODBias, ObjectInfo->VisibleMipsTilesCalculations, ObjectInfo->MeshHorizontalRange );
+					AddObject(Owner,ObjectInfo->MipMapLODBias, ObjectInfo->VisibleMipsTilesCalculations, ObjectInfo->MeshRange);
 				}
 			}
 		}
diff --git a/Engine/Plugins/Media/ImgMedia/Source/ImgMedia/Private/Assets/ImgMediaMipMapInfo.h b/Engine/Plugins/Media/ImgMedia/Source/ImgMedia/Private/Assets/ImgMediaMipMapInfo.h
--- a/Engine/Plugins/Media/ImgMedia/Source/ImgMedia/Private/Assets/ImgMediaMipMapInfo.h
+++ b/Engine/Plugins/Media/ImgMedia/Source/ImgMedia/Private/Assets/ImgMediaMipMapInfo.h
@@ -214,9 +214,11 @@ public:
 	 * @param InActor Object using our img sequence.
 	 * @param LODBias Bias used for mip-level calculations (Note: must be mirrored in a material manually).
 	 * @param MeshType Plane, sphere or none mesh type for mip-tile calculations.
-	 * @param MeshHorizontalRange Arc size in degrees used for visible tiles calculations, specific to the sphere.
+	 * @param MeshRange Arc size in degrees used for visible tiles calculations, specific to the sphere.
 	 */
-	void AddObject(AActor* InActor, float LODBias = 0.0f, EMediaTextureVisibleMipsTiles MeshType = EMediaTextureVisibleMipsTiles::None, float MeshHorizontalRange = 360.0f);
+	void AddObject(AActor* InActor, float LODBias = 0.0f,
+		EMediaTextureVisibleMipsTiles MeshType = EMediaTextureVisibleMipsTiles::None,
+		FVector2D MeshRange = FVector2D(360.0f, 180.0f));
 
 	/**
 	 * This object is no longer using our img sequence.
diff --git a/Engine/Plugins/Media/MediaPlate/Source/MediaPlate/Private/MediaPlateComponent.cpp b/Engine/Plugins/Media/MediaPlate/Source/MediaPlate/Private/MediaPlateComponent.cpp
--- a/Engine/Plugins/Media/MediaPlate/Source/MediaPlate/Private/MediaPlateComponent.cpp
+++ b/Engine/Plugins/Media/MediaPlate/Source/MediaPlate/Private/MediaPlateComponent.cpp
@@ -200,13 +200,13 @@ void UMediaPlateComponent::Stop()
 	StopClockSink();
 }
 
-void UMediaPlateComponent::SetMeshHorizontalRange(float InMeshHorizontalRange)
+void UMediaPlateComponent::SetMeshRange(FVector2D InMeshRange)
 {
-	MeshHorizontalRange = InMeshHorizontalRange;
+	MeshRange = InMeshRange;
 
 	if (MediaTextureTrackerObject != nullptr)
 	{
-		MediaTextureTrackerObject->MeshHorizontalRange = MeshHorizontalRange;
+		MediaTextureTrackerObject->MeshRange = MeshRange;
 	}
 }
 
@@ -217,7 +217,7 @@ void UMediaPlateComponent::RegisterWithMediaTextureTracker()
 	MediaTextureTrackerObject->Object = GetOwner();
 	MediaTextureTrackerObject->MipMapLODBias = 0.0f;
 	MediaTextureTrackerObject->VisibleMipsTilesCalculations = VisibleMipsTilesCalculations;
-	MediaTextureTrackerObject->MeshHorizontalRange = MeshHorizontalRange;
+	MediaTextureTrackerObject->MeshRange = MeshRange;
 
 	// Add our texture.
 	TObjectPtr<UMediaTexture> MediaTexture = GetMediaTexture();
diff --git a/Engine/Plugins/Media/MediaPlate/Source/MediaPlate/Public/MediaPlateComponent.h b/Engine/Plugins/Media/MediaPlate/Source/MediaPlate/Public/MediaPlateComponent.h
--- a/Engine/Plugins/Media/MediaPlate/Source/MediaPlate/Public/MediaPlateComponent.h
+++ b/Engine/Plugins/Media/MediaPlate/Source/MediaPlate/Public/MediaPlateComponent.h
@@ -107,13 +107,13 @@ public:
 	UPROPERTY(EditAnywhere, AdvancedDisplay, Category = "MediaPlate")
 	EMediaTextureVisibleMipsTiles VisibleMipsTilesCalculations;
 
-	/** Set the horizontal arc size in degrees used for visible mips and tiles calculations, specific to the sphere. */
+	/** Set the arc size in degrees used for visible mips and tiles calculations, specific to the sphere. */
 	UFUNCTION(BlueprintCallable, Category = "Media|MediaPlate")
-	void SetMeshHorizontalRange(float InMeshHorizontalRange);
+	void SetMeshRange(FVector2D InMeshRange);
 
-	/** Return the horizontal arc size in degrees used for visible mips and tiles calculations, specific to the sphere. */
+	/** Return the arc size in degrees used for visible mips and tiles calculations, specific to the sphere. */
 	UFUNCTION(BlueprintCallable, Category = "Media|MediaPlate")
-	float GetMeshHorizontalRange() const { return MeshHorizontalRange; }
+	FVector2D GetMeshRange() const { return MeshRange; }
 
 	/**
 	 * Adds our media texture to the media texture tracker.
@@ -141,7 +141,7 @@ public:
 
 private:
 	UPROPERTY()
-	float MeshHorizontalRange = 360.0f;
+	FVector2D MeshRange = FVector2D(360.0f, 180.0f);
 
 	/** Name for our media component. */
 	static FLazyName MediaComponentName;
diff --git a/Engine/Plugins/Media/MediaPlate/Source/MediaPlateEditor/Private/MediaPlateCustomization.cpp b/Engine/Plugins/Media/MediaPlate/Source/MediaPlateEditor/Private/MediaPlateCustomization.cpp
--- a/Engine/Plugins/Media/MediaPlate/Source/MediaPlateEditor/Private/MediaPlateCustomization.cpp
+++ b/Engine/Plugins/Media/MediaPlate/Source/MediaPlateEditor/Private/MediaPlateCustomization.cpp
@@ -452,6 +452,24 @@ void FMediaPlateCustomization::AddMeshCustomization(IDetailCategoryBuilder& Medi
 				.Value(this, &FMediaPlateCustomization::GetMeshHorizontalRange)
 				.OnValueChanged(this, &FMediaPlateCustomization::SetMeshHorizontalRange)
 		];
+
+	// Add sphere vertical arc.
+	DetailGroup.AddWidgetRow()
+		.Visibility(MeshSphereVisibility)
+		.NameContent()
+		[
+			SNew(STextBlock)
+				.Text(LOCTEXT("VerticalArc", "Vertical Arc"))
+				.ToolTipText(LOCTEXT("VerticalArc_ToolTip",
+				"Sets the vertical arc size of the sphere in degrees.\nFor example 180 for a half circle, 90 for a quarter circle."))
+				.Font(IDetailLayoutBuilder::GetDetailFont())
+		]
+		.ValueContent()
+		[
+			SNew(SNumericEntryBox<float>)
+				.Value(this, &FMediaPlateCustomization::GetMeshVerticalRange)
+				.OnValueChanged(this, &FMediaPlateCustomization::SetMeshVerticalRange)
+		];
 	
 }
 
@@ -615,23 +633,41 @@ TOptional<float> FMediaPlateCustomization::GetAspectRatio() const
 void FMediaPlateCustomization::SetMeshHorizontalRange(float HorizontalRange)
 {
 	HorizontalRange = FMath::Clamp(HorizontalRange, 0.0f, 360.0f);
+	TOptional VerticalRange = GetMeshVerticalRange();
+	if (VerticalRange.IsSet())
+	{
+		FVector2D MeshRange = FVector2D(HorizontalRange, VerticalRange.GetValue());
+		SetMeshRange(MeshRange);
+	}
+}
 
-	// Loop through all our objects.
+TOptional<float> FMediaPlateCustomization::GetMeshHorizontalRange() const
+{
+	// Loop through our objects.
 	for (const TWeakObjectPtr<UMediaPlateComponent>& MediaPlatePtr : MediaPlatesList)
 	{
 		UMediaPlateComponent* MediaPlate = MediaPlatePtr.Get();
 		if (MediaPlate != nullptr)
 		{
-			if (MediaPlate->GetMeshHorizontalRange() != HorizontalRange)
-			{
-				MediaPlate->SetMeshHorizontalRange(HorizontalRange);
-				SetSphereMesh(MediaPlate);
-			}
+			return MediaPlate->GetMeshRange().X;
 		}
 	}
+
+	return TOptional<float>();
 }
 
-TOptional<float> FMediaPlateCustomization::GetMeshHorizontalRange() const
+void FMediaPlateCustomization::SetMeshVerticalRange(float VerticalRange)
+{
+	VerticalRange = FMath::Clamp(VerticalRange, 0.0f, 180.0f);
+	TOptional HorizontalRange = GetMeshHorizontalRange();
+	if (HorizontalRange.IsSet())
+	{
+		FVector2D MeshRange = FVector2D(HorizontalRange.GetValue(), VerticalRange);
+		SetMeshRange(MeshRange);
+	}
+}
+
+TOptional<float> FMediaPlateCustomization::GetMeshVerticalRange() const
 {
 	// Loop through our objects.
 	for (const TWeakObjectPtr<UMediaPlateComponent>& MediaPlatePtr : MediaPlatesList)
@@ -639,13 +675,30 @@ TOptional<float> FMediaPlateCustomization::GetMeshHorizontalRange() const
 		UMediaPlateComponent* MediaPlate = MediaPlatePtr.Get();
 		if (MediaPlate != nullptr)
 		{
-			return MediaPlate->GetMeshHorizontalRange();
+			return MediaPlate->GetMeshRange().Y;
 		}
 	}
 
 	return TOptional<float>();
 }
 
+void FMediaPlateCustomization::SetMeshRange(FVector2D Range)
+{
+	// Loop through all our objects.
+	for (const TWeakObjectPtr<UMediaPlateComponent>& MediaPlatePtr : MediaPlatesList)
+	{
+		UMediaPlateComponent* MediaPlate = MediaPlatePtr.Get();
+		if (MediaPlate != nullptr)
+		{
+			if (MediaPlate->GetMeshRange() != Range)
+			{
+				MediaPlate->SetMeshRange(Range);
+				SetSphereMesh(MediaPlate);
+			}
+		}
+	}
+}
+
 FString FMediaPlateCustomization::GetMediaSourcePath() const
 {
 	FString Path;
diff --git a/Engine/Plugins/Media/MediaPlate/Source/MediaPlateEditor/Private/MediaPlateCustomization.h b/Engine/Plugins/Media/MediaPlate/Source/MediaPlateEditor/Private/MediaPlateCustomization.h
--- a/Engine/Plugins/Media/MediaPlate/Source/MediaPlateEditor/Private/MediaPlateCustomization.h
+++ b/Engine/Plugins/Media/MediaPlate/Source/MediaPlateEditor/Private/MediaPlateCustomization.h
@@ -111,6 +111,21 @@ private:
 	 */
 	TOptional<float> GetMeshHorizontalRange() const;
 
+	/**
+	 * Call this to set the vertical range of the mesh.
+	 */
+	void SetMeshVerticalRange(float VerticalRange);
+
+	/**
+	 * Call this to get the vertical range of the mesh.
+	 */
+	TOptional<float> GetMeshVerticalRange() const;
+
+	/**
+	 * Call this to set the range of the mesh.
+	 */
+	void SetMeshRange(FVector2D Range);
+
 	/**
 	 * Gets the object path for the media source object.
 	 */
diff --git a/Engine/Plugins/Media/MediaPlate/Source/MediaPlateEditor/Private/MediaPlateCustomizationMesh.cpp b/Engine/Plugins/Media/MediaPlate/Source/MediaPlateEditor/Private/MediaPlateCustomizationMesh.cpp
--- a/Engine/Plugins/Media/MediaPlate/Source/MediaPlateEditor/Private/MediaPlateCustomizationMesh.cpp
+++ b/Engine/Plugins/Media/MediaPlate/Source/MediaPlateEditor/Private/MediaPlateCustomizationMesh.cpp
@@ -139,7 +139,8 @@ void FMediaPlateCustomizationMesh::GenerateSphereMesh(FDynamicMesh3* OutMesh,
 {
 	FMediaPlateSphereGenerator SphereGen;
 	SphereGen.Radius = 50.0f;
-	SphereGen.ThetaRange = FMath::DegreesToRadians(MediaPlate->GetMeshHorizontalRange());
+	SphereGen.ThetaRange = FMath::DegreesToRadians(MediaPlate->GetMeshRange().X);
+	SphereGen.PhiRange = FMath::DegreesToRadians(MediaPlate->GetMeshRange().Y);
 	SphereGen.NumTheta = 64;
 	SphereGen.NumPhi = 64;
 	SphereGen.bPolygroupPerQuad = false;
@@ -183,10 +184,12 @@ UStaticMesh* FMediaPlateCustomizationMesh::CreateStaticMeshAsset(FDynamicMesh3*
 FString FMediaPlateCustomizationMesh::GetAssetPath(UMediaPlateComponent* MediaPlate)
 {
 	// Add the horizontal range.
-	FString ID = FString::SanitizeFloat(MediaPlate->GetMeshHorizontalRange());
-	ID.ReplaceCharInline(TCHAR('.'), TCHAR('_'));
+	FString ID_X = FString::SanitizeFloat(MediaPlate->GetMeshRange().X);
+	ID_X.ReplaceCharInline(TCHAR('.'), TCHAR('_'));
+	FString ID_Y = FString::SanitizeFloat(MediaPlate->GetMeshRange().Y);
+	ID_Y.ReplaceCharInline(TCHAR('.'), TCHAR('_'));
 
-	FString AssetPath = FString::Printf(TEXT("/Game/_MediaPlate/Sphere_%s"), *ID);
+	FString AssetPath = FString::Printf(TEXT("/Game/_MediaPlate/Sphere_%s_%s"), *ID_X, *ID_Y);
 
 	return AssetPath;
 }
diff --git a/Engine/Plugins/Media/MediaPlate/Source/MediaPlateEditor/Private/MediaPlateSphereGenerator.h b/Engine/Plugins/Media/MediaPlate/Source/MediaPlateEditor/Private/MediaPlateSphereGenerator.h
--- a/Engine/Plugins/Media/MediaPlate/Source/MediaPlateEditor/Private/MediaPlateSphereGenerator.h
+++ b/Engine/Plugins/Media/MediaPlate/Source/MediaPlateEditor/Private/MediaPlateSphereGenerator.h
@@ -14,6 +14,8 @@ class FMediaPlateSphereGenerator : public UE::Geometry::FMeshShapeGenerator
 public:
 	/** Radius */
 	double Radius = 1;
+	/** Range for the vertical arc. 180 degrees for a sphere, etc... */
+	double PhiRange = FMathd::Pi;
 	/** Range for the horizontal arc. 360 degrees for a sphere, etc... */
 	double ThetaRange = FMathd::TwoPi;
 
@@ -26,6 +28,8 @@ public:
 private:
 	/** True if this is a closed mesh, e.g. a sphere. */
 	bool bIsClosed = false;
+	/** True if we have the north and south poles. */
+	bool bHasPoles = false;
 
 	static FVector3d SphericalToCartesian(double r, double theta, double phi)
 	{
@@ -48,7 +52,7 @@ private:
 		};
 		{
 			int32 ThetaDivide = bIsClosed ? NumTheta : NumTheta -1;
-			const double Dphi = FMathd::Pi / double(NumPhi - 1);
+			const double Dphi = PhiRange / double(NumPhi - 1);
 			const double Dtheta = ThetaRange / double(ThetaDivide);
 
 			int32 p,t;
@@ -56,7 +60,19 @@ private:
 			int32 VtxIdx = 0;
 
 			// Add points between the poles
-			for (p = 1, Phi = Dphi; p < NumPhi - 1; ++p, Phi += Dphi) // NB: this skips the poles.
+			int32 pStart = 1;
+			int32 pEnd = NumPhi - 1;
+			double PhiStart = Dphi;
+
+			// If we dont have poles then add another layer top/bottom.
+			if (bHasPoles == false)
+			{
+				pStart--;
+				pEnd++;
+				PhiStart = (FMathd::Pi - PhiRange) * 0.5;
+			}
+
+			for (p = pStart, Phi = PhiStart; p < pEnd; ++p, Phi += Dphi) // NB: this skips the poles.
 			{
 				for (t = 0, Theta = 0; t < NumTheta; ++t, ++VtxIdx, Theta += Dtheta)
 				{
@@ -64,10 +80,14 @@ private:
 					SetVertex(VtxIdx, Normal * Radius, FVector3f(-Normal));
 				}
 			}
-			// add a single point at the North Pole
-			SetVertex(VtxIdx++, FVector3d::UnitZ() * Radius, -FVector3f::UnitZ());
-			// add a single point at the South Pole
-			SetVertex(VtxIdx++, -FVector3d::UnitZ() * Radius, FVector3f::UnitZ());
+
+			if (bHasPoles)
+			{
+				// add a single point at the North Pole
+				SetVertex(VtxIdx++, FVector3d::UnitZ() * Radius, -FVector3f::UnitZ());
+				// add a single point at the South Pole
+				SetVertex(VtxIdx++, -FVector3d::UnitZ() * Radius, FVector3f::UnitZ());
+			}
 		}
 	}
 
@@ -81,7 +101,19 @@ private:
 		int32 UVIdx = 0;
 		int32 p,t;
 		float UVPhi, UVTheta;
-		for ( p = 1, UVPhi = DUVphi; p < NumPhi - 1; ++p, UVPhi += DUVphi)
+		float UVPhiStart = DUVphi;
+		int32 pStart = 1;
+		int32 pEnd = NumPhi - 1;
+
+		// If we dont have poles then add another layer top/bottom.
+		if (bHasPoles == false)
+		{
+			UVPhiStart = 0.0f;
+			pStart--;
+			pEnd++;
+		}
+		
+		for ( p = pStart, UVPhi = UVPhiStart; p < pEnd; ++p, UVPhi += DUVphi)
 		{
 			for (t = 0, UVTheta = 1; t < NumTheta; ++t, ++UVIdx, UVTheta += DUVtheta)
 			{
@@ -92,17 +124,21 @@ private:
 			UVParentVertex[UVIdx] = (p - 1) * NumTheta; // Wrap around
 			++UVIdx;
 		}
-		int32 NorthPoleVtxIdx = (NumPhi - 2) * NumTheta;
-		for (t = 0, UVTheta = 1 + DUVtheta; t < NumTheta; ++t, ++UVIdx, UVTheta += DUVtheta)
-		{
-			UVs[UVIdx] = FVector2f(UVTheta, 0.0);
-			UVParentVertex[UVIdx] = NorthPoleVtxIdx;
-		}
-		int32 SouthPoleVtxIdx = NorthPoleVtxIdx + 1;
-		for (t = 0, UVTheta = 1 + DUVtheta; t < NumTheta; ++t, ++UVIdx, UVTheta += DUVtheta)
+
+		if (bHasPoles)
 		{
-			UVs[UVIdx] = FVector2f(UVTheta, 1.0);
-			UVParentVertex[UVIdx] = SouthPoleVtxIdx;
+			int32 NorthPoleVtxIdx = (NumPhi - 2) * NumTheta;
+			for (t = 0, UVTheta = 1 + DUVtheta; t < NumTheta; ++t, ++UVIdx, UVTheta += DUVtheta)
+			{
+				UVs[UVIdx] = FVector2f(UVTheta, 0.0);
+				UVParentVertex[UVIdx] = NorthPoleVtxIdx;
+			}
+			int32 SouthPoleVtxIdx = NorthPoleVtxIdx + 1;
+			for (t = 0, UVTheta = 1 + DUVtheta; t < NumTheta; ++t, ++UVIdx, UVTheta += DUVtheta)
+			{
+				UVs[UVIdx] = FVector2f(UVTheta, 1.0);
+				UVParentVertex[UVIdx] = SouthPoleVtxIdx;
+			}
 		}
 	}
 
@@ -122,7 +158,17 @@ private:
 		// Generate equatorial triangles
 		int32 Corners[4] =   { 0, 1,     NumTheta + 1, NumTheta};
 		int32 UVCorners[4] = { 0, 1, NumTheta + 2, NumTheta + 1};
-		for (int32 p = 1; p < NumPhi - 2; ++p)
+		int32 pStart = 1;
+		int32 pEnd = NumPhi - 2;
+
+		// If we dont have poles then add another layer top/bottom.
+		if (bHasPoles == false)
+		{
+			pStart--;
+			pEnd++;
+		}
+
+		for (int32 p = pStart; p < pEnd; ++p)
 		{
 			for (int32 t = 0; t < NumTheta - 1; ++t)
 			{
@@ -214,19 +260,35 @@ public:
 		{
 			ThetaRange = FMathd::TwoPi;
 		}
+
+		// Do we have the poles?
+		bHasPoles = PhiRange >= FMath::DegreesToRadians(179.9f);
+		if (bHasPoles)
+		{
+			PhiRange = FMathd::Pi;
+		}
 		
 		// enforce sane values for vertex counts
 		NumPhi = FMath::Max(NumPhi, 3);
 		NumTheta = FMath::Max(NumTheta, 3);
-		const int32 NumVertices = (NumPhi - 2) * NumTheta + 2;
-		const int32 NumUVs = (NumPhi - 2) * (NumTheta + 1) + (2 * NumTheta);
-		const int32 NumTris = (NumPhi - 2) * (bIsClosed ? NumTheta : (NumTheta - 1)) * 2;
+		int32 NumVertices = (NumPhi - 2) * NumTheta + 2;
+		int32 NumUVs = (NumPhi - 2) * (NumTheta + 1) + (2 * NumTheta);
+		int32 NumTris = (NumPhi - 2) * (bIsClosed ? NumTheta : (NumTheta - 1)) * 2;
+		if (bHasPoles == false)
+		{
+			NumVertices = NumPhi * NumTheta;
+			NumUVs = NumPhi * (NumTheta + 1);
+			NumTris = (NumPhi - 1) * (bIsClosed ? NumTheta : (NumTheta - 1)) * 2;
+		}
 		SetBufferSizes(NumVertices, NumTris, NumUVs, NumVertices);
 
 		GenerateVertices();
 		GenerateUVVertices();
 		OutputEquatorialTriangles();
-		OutputPolarTriangles();
+		if (bHasPoles)
+		{
+			OutputPolarTriangles();
+		}
 		return *this;
 	}
 
diff --git a/Engine/Source/Runtime/MediaAssets/Public/MediaTextureTracker.h b/Engine/Source/Runtime/MediaAssets/Public/MediaTextureTracker.h
--- a/Engine/Source/Runtime/MediaAssets/Public/MediaTextureTracker.h
+++ b/Engine/Source/Runtime/MediaAssets/Public/MediaTextureTracker.h
@@ -30,8 +30,8 @@ struct FMediaTextureTrackerObject
 	float MipMapLODBias;
 	/** Specify type of mesh used for visible mips and tiles calculations. */
 	EMediaTextureVisibleMipsTiles VisibleMipsTilesCalculations;
-	/** Horizontal arc size in degrees used for visible tiles calculations, specific to the sphere.*/
-	float MeshHorizontalRange;
+	/** Arc size in degrees used for visible tiles calculations, specific to the sphere.*/
+	FVector2D MeshRange;
 };
 
 /**
