MediaPlate: Fixed automatic aspect ratio not working sometimes.

#rnx

diff --git a/Engine/Plugins/Media/MediaPlate/Source/MediaPlate/Private/MediaPlateComponent.cpp b/Engine/Plugins/Media/MediaPlate/Source/MediaPlate/Private/MediaPlateComponent.cpp
--- a/Engine/Plugins/Media/MediaPlate/Source/MediaPlate/Private/MediaPlateComponent.cpp
+++ b/Engine/Plugins/Media/MediaPlate/Source/MediaPlate/Private/MediaPlateComponent.cpp
@@ -182,7 +182,6 @@ void UMediaPlateComponent::Play()
 				{
 					ClockSink = MakeShared<FMediaComponentClockSink, ESPMode::ThreadSafe>(this);
 				}
-				bIsWaitingForRender = true;
 				MediaModule->GetClock().AddSink(ClockSink.ToSharedRef());
 			}
 		}
@@ -292,28 +291,15 @@ void UMediaPlateComponent::TickOutput()
 		// Is the player ready?
 		if (MediaPlayer->IsPreparing() == false)
 		{
-			if (bIsWaitingForRender)
-			{
-				bIsWaitingForRender = false;
-			}
-			else
-			{
-				// Get the texture.
-				UMediaTexture* Texture = GetMediaTexture();
-				if (Texture != nullptr)
-				{
-					float Width = Texture->GetSurfaceWidth();
-					float Height = Texture->GetSurfaceHeight();
-					// Set aspect ratio.
-					if (Height != 0.0f)
-					{
-						float AspectRatio = Width / Height;
-						SetAspectRatio(AspectRatio);
-					}
-
-					// No need to tick anymore.
-					StopClockSink();
-				}
+			FIntPoint VideoDim = MediaPlayer->GetVideoTrackDimensions(INDEX_NONE, INDEX_NONE);
+			if (VideoDim.Y != 0)
+			{ 
+				// Set aspect ratio.
+				float AspectRatio = (float)VideoDim.X / (float)VideoDim.Y;
+				SetAspectRatio(AspectRatio);
+					
+				// No need to tick anymore.
+				StopClockSink();
 			}
 		}
 	}
diff --git a/Engine/Plugins/Media/MediaPlate/Source/MediaPlate/Public/MediaPlateComponent.h b/Engine/Plugins/Media/MediaPlate/Source/MediaPlate/Public/MediaPlateComponent.h
--- a/Engine/Plugins/Media/MediaPlate/Source/MediaPlate/Public/MediaPlateComponent.h
+++ b/Engine/Plugins/Media/MediaPlate/Source/MediaPlate/Public/MediaPlateComponent.h
@@ -152,8 +152,6 @@ private:
 	TSharedPtr<FMediaTextureTrackerObject, ESPMode::ThreadSafe> MediaTextureTrackerObject;
 	/** Our media clock sink. */
 	TSharedPtr<FMediaComponentClockSink, ESPMode::ThreadSafe> ClockSink;
-	/** True if we are waiting for the first frame to come out. */
-	bool bIsWaitingForRender = false;
 
 	/**
 	 * Plays a media source.
