ImgMedia: Fix spherical mesh UVs so that horizontal content isn't flipped from within the sphere.

#rnx

diff --git a/Engine/Plugins/Media/ImgMedia/Source/ImgMedia/Private/Assets/ImgMediaMipMapInfo.cpp b/Engine/Plugins/Media/ImgMedia/Source/ImgMedia/Private/Assets/ImgMediaMipMapInfo.cpp
--- a/Engine/Plugins/Media/ImgMedia/Source/ImgMedia/Private/Assets/ImgMediaMipMapInfo.cpp
+++ b/Engine/Plugins/Media/ImgMedia/Source/ImgMedia/Private/Assets/ImgMediaMipMapInfo.cpp
@@ -580,7 +580,7 @@ namespace {
 				// Convert from latlong UV to spherical coordinates
 				FVector2D TileCornerSpherical = FVector2D(
 					FMath::DegreesToRadians(MeshRange.Y) * InUV.Y,
-					FMath::DegreesToRadians(MeshRange.X) * (1.0f - InUV.X));
+					FMath::DegreesToRadians(MeshRange.X) * InUV.X);
 				FVector CornersWS = TileCornerSpherical.SphericalToUnitCartesian() * DefaultSphereRadius;
 				return MeshTransform.TransformPosition(CornersWS);
 			};
diff --git a/Engine/Plugins/Media/MediaPlate/Source/MediaPlateEditor/Private/MediaPlateSphereGenerator.h b/Engine/Plugins/Media/MediaPlate/Source/MediaPlateEditor/Private/MediaPlateSphereGenerator.h
--- a/Engine/Plugins/Media/MediaPlate/Source/MediaPlateEditor/Private/MediaPlateSphereGenerator.h
+++ b/Engine/Plugins/Media/MediaPlate/Source/MediaPlateEditor/Private/MediaPlateSphereGenerator.h
@@ -96,7 +96,7 @@ private:
 		// generate the UV's
 		int32 ThetaDivide = bIsClosed ? NumTheta : NumTheta - 1;
 		const float DUVphi = 1.0 / float(NumPhi - 1);
-		const float DUVtheta = -1.0 / float(ThetaDivide);
+		const float DUVtheta = 1.0 / float(ThetaDivide);
 
 		int32 UVIdx = 0;
 		int32 p,t;
@@ -115,7 +115,7 @@ private:
 		
 		for ( p = pStart, UVPhi = UVPhiStart; p < pEnd; ++p, UVPhi += DUVphi)
 		{
-			for (t = 0, UVTheta = 1; t < NumTheta; ++t, ++UVIdx, UVTheta += DUVtheta)
+			for (t = 0, UVTheta = 0; t < NumTheta; ++t, ++UVIdx, UVTheta += DUVtheta)
 			{
 				UVs[UVIdx] = FVector2f(UVTheta, UVPhi);
 				UVParentVertex[UVIdx] = (p - 1) * NumTheta + t;
@@ -128,13 +128,13 @@ private:
 		if (bHasPoles)
 		{
 			int32 NorthPoleVtxIdx = (NumPhi - 2) * NumTheta;
-			for (t = 0, UVTheta = 1 + DUVtheta; t < NumTheta; ++t, ++UVIdx, UVTheta += DUVtheta)
+			for (t = 0, UVTheta = DUVtheta; t < NumTheta; ++t, ++UVIdx, UVTheta += DUVtheta)
 			{
 				UVs[UVIdx] = FVector2f(UVTheta, 0.0);
 				UVParentVertex[UVIdx] = NorthPoleVtxIdx;
 			}
 			int32 SouthPoleVtxIdx = NorthPoleVtxIdx + 1;
-			for (t = 0, UVTheta = 1 + DUVtheta; t < NumTheta; ++t, ++UVIdx, UVTheta += DUVtheta)
+			for (t = 0, UVTheta = DUVtheta; t < NumTheta; ++t, ++UVIdx, UVTheta += DUVtheta)
 			{
 				UVs[UVIdx] = FVector2f(UVTheta, 1.0);
 				UVParentVertex[UVIdx] = SouthPoleVtxIdx;
