MediaPlate: Material is now an instance and not dynamic. The media texture (and player) are now non transient to support this.

#rnx

diff --git a/Engine/Plugins/Media/MediaPlate/Source/MediaPlate/Private/MediaPlate.cpp b/Engine/Plugins/Media/MediaPlate/Source/MediaPlate/Private/MediaPlate.cpp
--- a/Engine/Plugins/Media/MediaPlate/Source/MediaPlate/Private/MediaPlate.cpp
+++ b/Engine/Plugins/Media/MediaPlate/Source/MediaPlate/Private/MediaPlate.cpp
@@ -4,7 +4,7 @@
 
 #include "Components/StaticMeshComponent.h"
 #include "Engine/StaticMesh.h"
-#include "Materials/MaterialInstanceDynamic.h"
+#include "Materials/MaterialInstanceConstant.h"
 #include "MediaPlateComponent.h"
 #include "MediaPlateModule.h"
 #include "MediaTexture.h"
@@ -44,38 +44,26 @@ AMediaPlate::AMediaPlate(const FObjectInitializer& ObjectInitializer)
 	}
 }
 
+
+void AMediaPlate::PostActorCreated()
+{
+	Super::PostActorCreated();
+
+#if WITH_EDITOR
+	if (GIsEditor)
+	{
+		// Set which material to use.
+		UseDefaultMaterial();
+	}
+#endif
+}
+
 void AMediaPlate::PostRegisterAllComponents()
 {
 	Super::PostRegisterAllComponents();
 
 	if (MediaPlateComponent != nullptr)
 	{
-		if (StaticMeshComponent != nullptr)
-		{
-			UMaterialInstanceDynamic* MediaMID = nullptr;
-
-			// Do we have a material?
-			if (StaticMeshComponent->GetNumOverrideMaterials() > 0)
-			{
-				MediaMID = StaticMeshComponent->CreateAndSetMaterialInstanceDynamicFromMaterial(0, StaticMeshComponent->GetMaterial(0));
-			}
-			else
-			{
-				// Add default material.
-				UMaterial* Material = LoadObject<UMaterial>(nullptr, TEXT("/MediaPlate/M_MediaPlate"));
-				if (Material != nullptr)
-				{
-					MediaMID = StaticMeshComponent->CreateAndSetMaterialInstanceDynamicFromMaterial(0, Material);
-				}
-			}
-
-			// Set up the material to point to our media texture.
-			if (MediaMID != nullptr)
-			{
-				MediaMID->SetTextureParameterValue(MediaTextureName, MediaPlateComponent->GetMediaTexture());
-			}
-		}
-
 		// Add our media texture to the tracker.
 		MediaPlateComponent->RegisterWithMediaTextureTracker();
 	}
@@ -92,4 +80,41 @@ void AMediaPlate::BeginDestroy()
 	Super::BeginDestroy();
 }
 
+#if WITH_EDITOR
+
+void AMediaPlate::UseDefaultMaterial()
+{
+	// Get material.
+	UMaterial* Material = LoadObject<UMaterial>(NULL, TEXT("/MediaPlate/M_MediaPlate"), NULL, LOAD_None, NULL);
+	if (Material != nullptr)
+	{
+		// Change M_ to MI_ in material name and then generate a unique one.
+		FString MaterialName = Material->GetName();
+		if (MaterialName.StartsWith(TEXT("M_")))
+		{
+			MaterialName.InsertAt(1, TEXT("I"));
+		}
+		FName MaterialUniqueName = MakeUniqueObjectName(this, UMaterialInstanceConstant::StaticClass(),
+			FName(*MaterialName));
+
+		// Create instance.
+		UMaterialInstanceConstant* MaterialInstance =
+			NewObject<UMaterialInstanceConstant>(this, MaterialUniqueName);
+		MaterialInstance->SetParentEditorOnly(Material);
+		MaterialInstance->SetTextureParameterValueEditorOnly(
+			FMaterialParameterInfo(MediaTextureName),
+			MediaPlateComponent->GetMediaTexture());
+		MaterialInstance->PostEditChange();
+
+		// Update static mesh.
+		if (StaticMeshComponent != nullptr)
+		{
+			StaticMeshComponent->Modify();
+			StaticMeshComponent->SetMaterial(0, MaterialInstance);
+		}
+	}
+}
+
+#endif
+
 #undef LOCTEXT_NAMESPACE
diff --git a/Engine/Plugins/Media/MediaPlate/Source/MediaPlate/Private/MediaPlateComponent.cpp b/Engine/Plugins/Media/MediaPlate/Source/MediaPlate/Private/MediaPlateComponent.cpp
--- a/Engine/Plugins/Media/MediaPlate/Source/MediaPlate/Private/MediaPlateComponent.cpp
+++ b/Engine/Plugins/Media/MediaPlate/Source/MediaPlate/Private/MediaPlateComponent.cpp
@@ -63,18 +63,6 @@ UMediaPlateComponent::UMediaPlateComponent(const FObjectInitializer& ObjectIniti
 {
 	CacheSettings.bOverride = true;
 
-	// Set up media component.
-	MediaComponent = CreateDefaultSubobject<UMediaComponent>(MediaComponentName);
-	if (MediaComponent != nullptr)
-	{
-		// Set up media texture.
-		UMediaTexture* MediaTexture = MediaComponent->GetMediaTexture();
-		if (MediaTexture != nullptr)
-		{
-			MediaTexture->NewStyleOutput = true;
-		}
-	}
-
 	// Set up playlist.
 	MediaPlaylist = CreateDefaultSubobject<UMediaPlaylist>(MediaPlaylistName);
 
@@ -86,10 +74,31 @@ void UMediaPlateComponent::OnRegister()
 {
 	Super::OnRegister();
 
+	// Create media texture if we don't have one.
+	if (MediaTexture == nullptr)
+	{
+		MediaTexture = NewObject<UMediaTexture>(this);
+		MediaTexture->NewStyleOutput = true;
+	}
+
+	// Create media player if we don't have one.
+	if (MediaPlayer == nullptr)
+	{
+		MediaPlayer = NewObject<UMediaPlayer>(this);
+		MediaPlayer->SetLooping(false);
+		MediaPlayer->PlayOnOpen = false;
+	}
+
+	// Set up media texture.
+	if (MediaTexture != nullptr)
+	{
+		MediaTexture->SetMediaPlayer(MediaPlayer);
+		MediaTexture->UpdateResource();
+	}
+
 	// Set up sound component if we have one.
 	if (SoundComponent != nullptr)
 	{
-		TObjectPtr<UMediaPlayer> MediaPlayer = GetMediaPlayer();
 		if (MediaPlayer != nullptr)
 		{
 			SoundComponent->SetMediaPlayer(MediaPlayer);
@@ -129,30 +138,17 @@ void UMediaPlateComponent::BeginDestroy()
 
 UMediaPlayer* UMediaPlateComponent::GetMediaPlayer()
 {
-	TObjectPtr<UMediaPlayer> MediaPlayer = nullptr;
-	if (MediaComponent != nullptr)
-	{
-		MediaPlayer = MediaComponent->GetMediaPlayer();
-	}
-
 	return MediaPlayer;
 }
 
 UMediaTexture* UMediaPlateComponent::GetMediaTexture()
 {
-	TObjectPtr<UMediaTexture> MediaTexture = nullptr;
-	if (MediaComponent != nullptr)
-	{
-		MediaTexture = MediaComponent->GetMediaTexture();
-	}
-
 	return MediaTexture;
 }
 
 void UMediaPlateComponent::Play()
 {
 	bool bIsPlaying = false;
-	TObjectPtr<UMediaPlayer> MediaPlayer = GetMediaPlayer();
 	if (MediaPlayer != nullptr)
 	{
 		UMediaSource* MediaSource = nullptr;
@@ -190,7 +186,6 @@ void UMediaPlateComponent::Play()
 
 void UMediaPlateComponent::Stop()
 {
-	TObjectPtr<UMediaPlayer> MediaPlayer = GetMediaPlayer();
 	if (MediaPlayer != nullptr)
 	{
 		MediaPlayer->Close();
@@ -219,7 +214,6 @@ void UMediaPlateComponent::RegisterWithMediaTextureTracker()
 	MediaTextureTrackerObject->MeshRange = MeshRange;
 
 	// Add our texture.
-	TObjectPtr<UMediaTexture> MediaTexture = GetMediaTexture();
 	if (MediaTexture != nullptr)
 	{
 		FMediaTextureTracker& MediaTextureTracker = FMediaTextureTracker::Get();
@@ -233,7 +227,6 @@ void UMediaPlateComponent::UnregisterWithMediaTextureTracker()
 	if (MediaTextureTrackerObject != nullptr)
 	{
 		FMediaTextureTracker& MediaTextureTracker = FMediaTextureTracker::Get();
-		TObjectPtr<UMediaTexture> MediaTexture = GetMediaTexture();
 		MediaTextureTracker.UnregisterTexture(MediaTextureTrackerObject, MediaTexture);
 	}
 }
@@ -248,7 +241,6 @@ bool UMediaPlateComponent::PlayMediaSource(UMediaSource* InMediaSource)
 		InMediaSource->SetCacheSettings(CacheSettings);
 		
 		// Set media options.
-		TObjectPtr<UMediaPlayer> MediaPlayer = GetMediaPlayer();
 		if (MediaPlayer != nullptr)
 		{
 			// Play the source.
@@ -285,7 +277,6 @@ void UMediaPlateComponent::SetAspectRatio(float AspectRatio)
 
 void UMediaPlateComponent::TickOutput()
 {
-	TObjectPtr<UMediaPlayer> MediaPlayer = GetMediaPlayer();
 	if (MediaPlayer != nullptr)
 	{
 		// Is the player ready?
@@ -326,7 +317,6 @@ void UMediaPlateComponent::OnVisibleMipsTilesCalculationsChange()
 		MediaTextureTrackerObject->VisibleMipsTilesCalculations = VisibleMipsTilesCalculations;
 
 		// Propagate the change by restarting the player if it is currently playing.
-		TObjectPtr<UMediaPlayer> MediaPlayer = GetMediaPlayer();
 		if (MediaPlayer != nullptr)
 		{
 			if (MediaPlayer->IsPlaying())
@@ -347,7 +337,6 @@ void UMediaPlateComponent::PostEditChangeProperty(FPropertyChangedEvent& Propert
 		if (bEnableAudio)
 		{
 			// Get the media player.
-			TObjectPtr<UMediaPlayer> MediaPlayer = GetMediaPlayer();
 			if (MediaPlayer != nullptr)
 			{
 				// Create a sound component.
diff --git a/Engine/Plugins/Media/MediaPlate/Source/MediaPlate/Public/MediaPlate.h b/Engine/Plugins/Media/MediaPlate/Source/MediaPlate/Public/MediaPlate.h
--- a/Engine/Plugins/Media/MediaPlate/Source/MediaPlate/Public/MediaPlate.h
+++ b/Engine/Plugins/Media/MediaPlate/Source/MediaPlate/Public/MediaPlate.h
@@ -18,6 +18,7 @@ class MEDIAPLATE_API AMediaPlate : public AActor
 
 public:
 	//~ Begin AActor Interface
+	virtual void PostActorCreated();
 	virtual void PostRegisterAllComponents() override;
 	virtual void BeginDestroy() override;
 	//~ End AActor Interface
@@ -29,6 +30,15 @@ public:
 	UPROPERTY(Category = MediaPlate, VisibleAnywhere, BlueprintReadOnly)
 	TObjectPtr<UStaticMeshComponent> StaticMeshComponent;
 
+#if WITH_EDITOR
+
+	/*
+	 * Call this to change the static mesh to use the default media plate material.
+	 */
+	void UseDefaultMaterial();
+
+#endif
+
 private:
 	/** Name for our media plate component. */
 	static FLazyName MediaPlateComponentName;
diff --git a/Engine/Plugins/Media/MediaPlate/Source/MediaPlate/Public/MediaPlateComponent.h b/Engine/Plugins/Media/MediaPlate/Source/MediaPlate/Public/MediaPlateComponent.h
--- a/Engine/Plugins/Media/MediaPlate/Source/MediaPlate/Public/MediaPlateComponent.h
+++ b/Engine/Plugins/Media/MediaPlate/Source/MediaPlate/Public/MediaPlateComponent.h
@@ -80,9 +80,7 @@ public:
 	UPROPERTY(BlueprintReadWrite, Category = "MediaPlate")
 	bool bIsAspectRatioAuto = true;
 
-	/** Holds the media player. */
-	UPROPERTY(Category = MediaPlate, VisibleAnywhere, BlueprintReadOnly)
-	TObjectPtr<UMediaComponent> MediaComponent;
+	
 
 	/** Holds the component to play sound. */
 	UPROPERTY(EditAnywhere, AdvancedDisplay, Category = MediaPlate)
@@ -148,6 +146,14 @@ private:
 	/** Name for our playlist. */
 	static FLazyName MediaPlaylistName;
 
+	/** Holds the media player. */
+	UPROPERTY()
+	TObjectPtr<UMediaTexture> MediaTexture;
+
+	/** This component's media player */
+	UPROPERTY()
+	TObjectPtr<UMediaPlayer> MediaPlayer;
+
 	/** Info representing this object. */
 	TSharedPtr<FMediaTextureTrackerObject, ESPMode::ThreadSafe> MediaTextureTrackerObject;
 	/** Our media clock sink. */
diff --git a/Engine/Plugins/Media/MediaPlate/Source/MediaPlateEditor/Private/MediaPlateEditorModule.cpp b/Engine/Plugins/Media/MediaPlate/Source/MediaPlateEditor/Private/MediaPlateEditorModule.cpp
--- a/Engine/Plugins/Media/MediaPlate/Source/MediaPlateEditor/Private/MediaPlateEditorModule.cpp
+++ b/Engine/Plugins/Media/MediaPlate/Source/MediaPlateEditor/Private/MediaPlateEditorModule.cpp
@@ -6,6 +6,7 @@
 #include "Editor.h"
 #include "ISequencerModule.h"
 #include "LevelEditorViewport.h"
+#include "MaterialList.h"
 #include "MediaPlateComponent.h"
 #include "MediaPlateCustomization.h"
 #include "MediaPlateEditorStyle.h"
@@ -17,6 +18,7 @@
 #include "Sequencer/MediaPlateTrackEditor.h"
 #include "SLevelViewport.h"
 #include "Subsystems/EditorAssetSubsystem.h"
+#include "Widgets/SMediaPlateEditorMaterial.h"
 
 DEFINE_LOG_CATEGORY(LogMediaPlateEditor);
 
@@ -38,11 +40,20 @@ void FMediaPlateEditorModule::StartupModule()
 	ISequencerModule& SequencerModule = FModuleManager::Get().LoadModuleChecked<ISequencerModule>("Sequencer");
 	TrackEditorBindingHandle = SequencerModule.RegisterPropertyTrackEditor<FMediaPlateTrackEditor>();
 
+	// Add bottom extender for material item
+	FMaterialList::OnAddMaterialItemViewExtraBottomWidget.AddLambda([](const TSharedRef<FMaterialItemView>& InMaterialItemView, UActorComponent* InCurrentComponent, IDetailLayoutBuilder& InDetailBuilder, TArray<TSharedPtr<SWidget>>& OutExtensions)
+	{
+		LLM_SCOPE_BYNAME("MediaPlate/MediaPlateEditor");
+		OutExtensions.Add(SNew(SMediaPlateEditorMaterial, InMaterialItemView, InCurrentComponent));
+	});
+
 	FCoreDelegates::OnPostEngineInit.AddRaw(this, &FMediaPlateEditorModule::OnPostEngineInit);
 }
 
 void FMediaPlateEditorModule::ShutdownModule()
 {
+	FMaterialList::OnAddMaterialItemViewExtraBottomWidget.RemoveAll(this);
+
 	if (GEditor != nullptr)
 	{
 		UEditorAssetSubsystem* EditorAssetSubsystem = GEditor->GetEditorSubsystem<UEditorAssetSubsystem>();
diff --git a/Engine/Plugins/Media/MediaPlate/Source/MediaPlateEditor/Private/Widgets/SMediaPlateEditorMaterial.cpp b/Engine/Plugins/Media/MediaPlate/Source/MediaPlateEditor/Private/Widgets/SMediaPlateEditorMaterial.cpp
new file mode 100644
--- /dev/null
+++ b/Engine/Plugins/Media/MediaPlate/Source/MediaPlateEditor/Private/Widgets/SMediaPlateEditorMaterial.cpp
@@ -0,0 +1,53 @@
+// Copyright Epic Games, Inc. All Rights Reserved.
+
+#include "Widgets/SMediaPlateEditorMaterial.h"
+
+#include "MediaPlate.h"
+#include "Widgets/Input/SButton.h"
+#include "Widgets/Layout/SBox.h"
+
+#define LOCTEXT_NAMESPACE "SMediaPlateEditorMaterial"
+
+/* SMediaPlateEditorMaterial interface
+ *****************************************************************************/
+
+void SMediaPlateEditorMaterial::Construct(const FArguments& InArgs, const TSharedRef<FMaterialItemView>& InMaterialItemView, UActorComponent* InCurrentComponent)
+{
+	TSharedRef<SVerticalBox> ResultWidget = SNew(SVerticalBox);
+
+	// Get media plate.
+	if (InCurrentComponent != nullptr)
+	{
+		MediaPlate = InCurrentComponent->GetOwner<AMediaPlate>();
+		if (MediaPlate != nullptr)
+		{
+			// Add button for default material.
+			ResultWidget->AddSlot()
+				.AutoHeight()
+				.VAlign(VAlign_Center)
+				[
+					SNew(SButton)
+						.Text(LOCTEXT("ApplyDefaultMaterial", "Use Media Plate Default"))
+						.ToolTipText(LOCTEXT("ApplyDefaultMaterialTooltip", "Set the material to the default Media Plate material."))
+						.OnClicked(this, &SMediaPlateEditorMaterial::OnDefaultButtonClicked)
+				];
+		}
+	}
+
+	ChildSlot
+	[
+		ResultWidget
+	];
+}
+
+FReply SMediaPlateEditorMaterial::OnDefaultButtonClicked() const
+{
+	if (MediaPlate != nullptr)
+	{
+		MediaPlate->UseDefaultMaterial();
+	}
+
+	return FReply::Handled();
+}
+
+#undef LOCTEXT_NAMESPACE
diff --git a/Engine/Plugins/Media/MediaPlate/Source/MediaPlateEditor/Private/Widgets/SMediaPlateEditorMaterial.h b/Engine/Plugins/Media/MediaPlate/Source/MediaPlateEditor/Private/Widgets/SMediaPlateEditorMaterial.h
new file mode 100644
--- /dev/null
+++ b/Engine/Plugins/Media/MediaPlate/Source/MediaPlateEditor/Private/Widgets/SMediaPlateEditorMaterial.h
@@ -0,0 +1,39 @@
+// Copyright Epic Games, Inc. All Rights Reserved.
+
+#pragma once
+
+#include "MaterialList.h"
+#include "Widgets/SCompoundWidget.h"
+
+class AMediaPlate;
+
+/**
+ * Widget for editing materials with media plate.
+ */
+class SMediaPlateEditorMaterial : public SCompoundWidget
+{
+public:
+
+	SLATE_BEGIN_ARGS(SMediaPlateEditorMaterial) { }
+	SLATE_END_ARGS()
+
+public:
+
+	/**
+	 * Construct this widget.
+	 *
+	 * @param InArgs				The declaration data for this widget.
+	 * @param InMaterialItemView	Material view.
+	 * @param InCurrentComponent	component that is using the material.
+	 */
+	void Construct(const FArguments& InArgs, const TSharedRef<FMaterialItemView>& InMaterialItemView, UActorComponent* InCurrentComponent);
+
+private:
+	/**
+	 * Callback for when we press the default button.
+	 */
+	FReply OnDefaultButtonClicked() const;
+
+	/** Weak pointer to current component using the current material */
+	TWeakObjectPtr<AMediaPlate> MediaPlate = nullptr;
+};
