MediaPlate: Disabled shadows.

#rnx

diff --git a/Engine/Plugins/Media/MediaPlate/Source/MediaPlate/Private/MediaPlate.cpp b/Engine/Plugins/Media/MediaPlate/Source/MediaPlate/Private/MediaPlate.cpp
--- a/Engine/Plugins/Media/MediaPlate/Source/MediaPlate/Private/MediaPlate.cpp
+++ b/Engine/Plugins/Media/MediaPlate/Source/MediaPlate/Private/MediaPlate.cpp
@@ -26,6 +26,8 @@ AMediaPlate::AMediaPlate(const FObjectInitializer& ObjectInitializer)
 	// Set up static mesh component.
 	StaticMeshComponent = CreateDefaultSubobject<UStaticMeshComponent>(TEXT("StaticMeshComponent"));
 	StaticMeshComponent->SetupAttachment(RootComponent);
+	StaticMeshComponent->bCastStaticShadow = false;
+	StaticMeshComponent->bCastDynamicShadow = false;
 	MediaPlateComponent->StaticMeshComponent = StaticMeshComponent;
 
 	// Hook up mesh.
