MediaPlate: New material overrides are now correctly reconnected with the media texture.

#rnx

diff --git a/Engine/Plugins/Media/MediaPlate/Source/MediaPlate/Private/MediaPlate.cpp b/Engine/Plugins/Media/MediaPlate/Source/MediaPlate/Private/MediaPlate.cpp
--- a/Engine/Plugins/Media/MediaPlate/Source/MediaPlate/Private/MediaPlate.cpp
+++ b/Engine/Plugins/Media/MediaPlate/Source/MediaPlate/Private/MediaPlate.cpp
@@ -86,8 +86,13 @@ void AMediaPlate::BeginDestroy()
 
 void AMediaPlate::UseDefaultMaterial()
 {
-	// Get material.
-	UMaterial* Material = LoadObject<UMaterial>(NULL, TEXT("/MediaPlate/M_MediaPlate"), NULL, LOAD_None, NULL);
+	UMaterial* DefaultMaterial = LoadObject<UMaterial>(NULL, TEXT("/MediaPlate/M_MediaPlate"), NULL, LOAD_None, NULL);
+	
+	ApplyMaterial(DefaultMaterial);
+}
+
+void AMediaPlate::ApplyMaterial(UMaterialInterface* Material)
+{
 	if (Material != nullptr)
 	{
 		// Change M_ to MI_ in material name and then generate a unique one.
diff --git a/Engine/Plugins/Media/MediaPlate/Source/MediaPlate/Public/MediaPlate.h b/Engine/Plugins/Media/MediaPlate/Source/MediaPlate/Public/MediaPlate.h
--- a/Engine/Plugins/Media/MediaPlate/Source/MediaPlate/Public/MediaPlate.h
+++ b/Engine/Plugins/Media/MediaPlate/Source/MediaPlate/Public/MediaPlate.h
@@ -7,6 +7,7 @@
 #include "MediaPlate.generated.h"
 
 class UMediaPlateComponent;
+class UMaterialInterface;
 
 /**
  * MediaPlate is an actor that can play and show media in the world.
@@ -36,6 +37,7 @@ public:
 	 * Call this to change the static mesh to use the default media plate material.
 	 */
 	void UseDefaultMaterial();
+	void ApplyMaterial(UMaterialInterface* Material);
 
 #endif
 
diff --git a/Engine/Plugins/Media/MediaPlate/Source/MediaPlateEditor/Private/MediaPlateCustomization.cpp b/Engine/Plugins/Media/MediaPlate/Source/MediaPlateEditor/Private/MediaPlateCustomization.cpp
--- a/Engine/Plugins/Media/MediaPlate/Source/MediaPlateEditor/Private/MediaPlateCustomization.cpp
+++ b/Engine/Plugins/Media/MediaPlate/Source/MediaPlateEditor/Private/MediaPlateCustomization.cpp
@@ -30,6 +30,42 @@
 /* IDetailCustomization interface
  *****************************************************************************/
 
+FMediaPlateCustomization::FMediaPlateCustomization()
+{
+	PropertyChangeDelegate = FCoreUObjectDelegates::OnObjectPropertyChanged.AddRaw(this, &FMediaPlateCustomization::OnObjectPropertyChanged);
+}
+
+FMediaPlateCustomization::~FMediaPlateCustomization()
+{
+	FCoreUObjectDelegates::OnObjectPropertyChanged.RemoveAll(this);
+}
+
+void FMediaPlateCustomization::OnObjectPropertyChanged(UObject* InObject, FPropertyChangedEvent& InEvent)
+{
+	if (InEvent.GetPropertyName() == GET_MEMBER_NAME_CHECKED(UMeshComponent, OverrideMaterials))
+	{
+		for (TWeakObjectPtr<UMediaPlateComponent>& MediaPlateComponent : MediaPlatesList)
+		{
+			if (MediaPlateComponent.IsValid())
+			{
+				if (InObject->GetOuter() == MediaPlateComponent->GetOuter())
+				{
+					AMediaPlate* MediaPlate = Cast<AMediaPlate>(InObject->GetOuter());
+
+					if (MediaPlate != nullptr && MediaPlate->StaticMeshComponent != nullptr)
+					{
+						const TObjectPtr<UMaterialInterface>& Material = MediaPlate->StaticMeshComponent->OverrideMaterials.Pop();
+
+						MediaPlate->ApplyMaterial(Material);
+					}
+
+					break;
+				}
+			}
+		}
+	}
+}
+
 void FMediaPlateCustomization::CustomizeDetails(IDetailLayoutBuilder& DetailBuilder)
 {
 	// Is this the media plate editor window?
diff --git a/Engine/Plugins/Media/MediaPlate/Source/MediaPlateEditor/Private/MediaPlateCustomization.h b/Engine/Plugins/Media/MediaPlate/Source/MediaPlateEditor/Private/MediaPlateCustomization.h
--- a/Engine/Plugins/Media/MediaPlate/Source/MediaPlateEditor/Private/MediaPlateCustomization.h
+++ b/Engine/Plugins/Media/MediaPlate/Source/MediaPlateEditor/Private/MediaPlateCustomization.h
@@ -24,6 +24,9 @@ class FMediaPlateCustomization
 {
 public:
 
+	FMediaPlateCustomization();
+	~FMediaPlateCustomization();
+
 	//~ IDetailCustomization interface
 
 	virtual void CustomizeDetails(IDetailLayoutBuilder& DetailBuilder) override;
@@ -40,6 +43,9 @@ public:
 
 private:
 
+	/** Property change delegate used for static mesh material changes. */	
+	FDelegateHandle PropertyChangeDelegate;
+
 	/** List of the media plates we are editing. */
 	TArray<TWeakObjectPtr<UMediaPlateComponent>> MediaPlatesList;
 	/** Stores the current value of the MediaPath property. */
@@ -51,6 +57,11 @@ private:
 	/** Handles mesh stuff. */
 	FMediaPlateCustomizationMesh MeshCustomization;
 
+	/**
+	 * Object property change callback used for static mesh material changes.
+	 */
+	void OnObjectPropertyChanged(UObject* InObject, FPropertyChangedEvent& InEvent);
+
 	/**
 	 * Adds widgets for editing the mesh.
 	 */
