MediaPlate: Added option to not play media when the plate is not visible.

#rnx

diff --git a/Engine/Plugins/Media/MediaPlate/Source/MediaPlate/Private/MediaPlateComponent.cpp b/Engine/Plugins/Media/MediaPlate/Source/MediaPlate/Private/MediaPlateComponent.cpp
--- a/Engine/Plugins/Media/MediaPlate/Source/MediaPlate/Private/MediaPlateComponent.cpp
+++ b/Engine/Plugins/Media/MediaPlate/Source/MediaPlate/Private/MediaPlateComponent.cpp
@@ -61,6 +61,9 @@ FLazyName UMediaPlateComponent::MediaPlaylistName(TEXT("MediaPlaylist0"));
 UMediaPlateComponent::UMediaPlateComponent(const FObjectInitializer& ObjectInitializer)
 	: Super(ObjectInitializer)
 {
+	PrimaryComponentTick.bCanEverTick = true;
+	PrimaryComponentTick.bStartWithTickEnabled = false;
+
 	CacheSettings.bOverride = true;
 
 	// Set up playlist.
@@ -110,6 +113,9 @@ void UMediaPlateComponent::BeginPlay()
 {
 	Super::BeginPlay();
 
+	// Activate tick if needed.
+	UpdateTicking();
+
 	// Start playing?
 	if (bAutoPlay)
 	{
@@ -136,6 +142,33 @@ void UMediaPlateComponent::BeginDestroy()
 	Super::BeginDestroy();
 }
 
+
+void UMediaPlateComponent::TickComponent(float DeltaTime, enum ELevelTick TickType, FActorComponentTickFunction* ThisTickFunction)
+{
+	Super::TickComponent(DeltaTime, TickType, ThisTickFunction);
+
+	if (bPlayOnlyWhenVisible)
+	{
+		bool bIsVisible = IsVisible();
+
+		if (MediaPlayer != nullptr)
+		{
+			if (bIsVisible)
+			{
+				ResumeWhenVisible();
+			}
+			else
+			{
+				if (MediaPlayer->IsPlaying())
+				{
+					MediaPlayer->Pause();
+					TimeWhenPlaybackPaused = FApp::GetGameTime();
+				}
+			}
+		}
+	}
+}
+
 UMediaPlayer* UMediaPlateComponent::GetMediaPlayer()
 {
 	return MediaPlayer;
@@ -148,40 +181,48 @@ UMediaTexture* UMediaPlateComponent::GetMediaTexture()
 
 void UMediaPlateComponent::Play()
 {
-	bool bIsPlaying = false;
-	if (MediaPlayer != nullptr)
+	if ((bPlayOnlyWhenVisible == false) || (IsVisible()))
 	{
-		UMediaSource* MediaSource = nullptr;
-		if (MediaPlaylist != nullptr)
+		bool bIsPlaying = false;
+			if (MediaPlayer != nullptr)
+			{
+				UMediaSource* MediaSource = nullptr;
+					if (MediaPlaylist != nullptr)
+					{
+						MediaSource = MediaPlaylist->Get(0);
+					}
+				bIsPlaying = PlayMediaSource(MediaSource);
+			}
+
+		// Did anything play?
+		if (bIsPlaying == false)
 		{
-			MediaSource = MediaPlaylist->Get(0);
+			UE_LOG(LogMediaPlate, Warning, TEXT("Could not play anything."));
 		}
-		bIsPlaying = PlayMediaSource(MediaSource);
-	}
-
-	// Did anything play?
-	if (bIsPlaying == false)
-	{
-		UE_LOG(LogMediaPlate, Warning, TEXT("Could not play anything."));
-	}
-	else
-	{
-		// Are we using automatic aspect ratio?
-		if ((bIsAspectRatioAuto) &&
-			(VisibleMipsTilesCalculations == EMediaTextureVisibleMipsTiles::Plane))
+		else
 		{
-			// Start the clock sink so we can tick.
-			IMediaModule* MediaModule = FModuleManager::LoadModulePtr<IMediaModule>("Media");
-			if (MediaModule != nullptr)
+			// Are we using automatic aspect ratio?
+			if ((bIsAspectRatioAuto) &&
+				(VisibleMipsTilesCalculations == EMediaTextureVisibleMipsTiles::Plane))
 			{
-				if (ClockSink.IsValid() == false)
+				// Start the clock sink so we can tick.
+				IMediaModule* MediaModule = FModuleManager::LoadModulePtr<IMediaModule>("Media");
+				if (MediaModule != nullptr)
 				{
-					ClockSink = MakeShared<FMediaComponentClockSink, ESPMode::ThreadSafe>(this);
+					if (ClockSink.IsValid() == false)
+					{
+						ClockSink = MakeShared<FMediaComponentClockSink, ESPMode::ThreadSafe>(this);
+					}
+					MediaModule->GetClock().AddSink(ClockSink.ToSharedRef());
 				}
-				MediaModule->GetClock().AddSink(ClockSink.ToSharedRef());
 			}
 		}
 	}
+	else
+	{
+		bWantsToPlayWhenVisible = true;
+		TimeWhenPlaybackPaused = FApp::GetGameTime();
+	}
 }
 
 void UMediaPlateComponent::Stop()
@@ -192,6 +233,7 @@ void UMediaPlateComponent::Stop()
 	}
 
 	StopClockSink();
+	bWantsToPlayWhenVisible = false;
 }
 
 void UMediaPlateComponent::SetMeshRange(FVector2D InMeshRange)
@@ -204,6 +246,18 @@ void UMediaPlateComponent::SetMeshRange(FVector2D InMeshRange)
 	}
 }
 
+void UMediaPlateComponent::SetPlayOnlyWhenVisible(bool bInPlayOnlyWhenVisible)
+{
+	// If we are turning off PlayOnlyWhenVisible then make sure we are playing.
+	if (bInPlayOnlyWhenVisible == false)
+	{
+		ResumeWhenVisible();
+	}
+
+	bPlayOnlyWhenVisible = bInPlayOnlyWhenVisible;
+	UpdateTicking();
+}
+
 void UMediaPlateComponent::RegisterWithMediaTextureTracker()
 {
 	// Set up object.
@@ -308,6 +362,59 @@ void UMediaPlateComponent::StopClockSink()
 	}
 }
 
+bool UMediaPlateComponent::IsVisible()
+{
+#if WITH_EDITOR
+	// Always return true if the game is not running as the visibility checks do not work in this
+	// case.
+	if (FApp::IsGame() == false)
+	{
+		return true;
+	}
+#endif
+
+	return GetOwner()->WasRecentlyRendered();
+}
+
+void UMediaPlateComponent::ResumeWhenVisible()
+{
+	if (MediaPlayer != nullptr)
+	{
+		if (MediaPlayer->IsPaused())
+		{
+			FTimespan PlayTime = GetResumeTime();
+			MediaPlayer->Seek(PlayTime);
+			MediaPlayer->Play();
+		}
+		else if (bWantsToPlayWhenVisible)
+		{
+			Play();
+			FTimespan PlayTime = GetResumeTime();
+			MediaPlayer->Seek(PlayTime);
+		}
+	}
+}
+
+FTimespan UMediaPlateComponent::GetResumeTime()
+{
+	FTimespan PlayerTime;
+	if (MediaPlayer != nullptr)
+	{
+		PlayerTime = MediaPlayer->GetTime();
+		float CurrentTime = FApp::GetGameTime();
+		float ElapsedTime = CurrentTime - TimeWhenPlaybackPaused;
+		PlayerTime += FTimespan::FromSeconds(ElapsedTime);
+	}
+
+	return PlayerTime;
+}
+
+void UMediaPlateComponent::UpdateTicking()
+{
+	bool bEnableTick = bPlayOnlyWhenVisible;
+	PrimaryComponentTick.SetTickFunctionEnable(bEnableTick);
+}
+
 #if WITH_EDITOR
 
 void UMediaPlateComponent::OnVisibleMipsTilesCalculationsChange()
diff --git a/Engine/Plugins/Media/MediaPlate/Source/MediaPlate/Public/MediaPlateComponent.h b/Engine/Plugins/Media/MediaPlate/Source/MediaPlate/Public/MediaPlateComponent.h
--- a/Engine/Plugins/Media/MediaPlate/Source/MediaPlate/Public/MediaPlateComponent.h
+++ b/Engine/Plugins/Media/MediaPlate/Source/MediaPlate/Public/MediaPlateComponent.h
@@ -32,6 +32,8 @@ public:
 	virtual void OnRegister();
 	virtual void BeginPlay();
 	virtual void BeginDestroy();
+	virtual void TickComponent(float DeltaTime, enum ELevelTick TickType, FActorComponentTickFunction* ThisTickFunction) override;
+
 #if WITH_EDITOR
 	virtual void PostEditChangeProperty(FPropertyChangedEvent& PropertyChangedEvent) override;
 #endif // WITH_EDITOR
@@ -80,8 +82,6 @@ public:
 	UPROPERTY(BlueprintReadWrite, Category = "MediaPlate")
 	bool bIsAspectRatioAuto = true;
 
-	
-
 	/** Holds the component to play sound. */
 	UPROPERTY(EditAnywhere, AdvancedDisplay, Category = MediaPlate)
 	TObjectPtr<UMediaSoundComponent> SoundComponent;
@@ -113,6 +113,10 @@ public:
 	UFUNCTION(BlueprintCallable, Category = "Media|MediaPlate")
 	FVector2D GetMeshRange() const { return MeshRange; }
 
+	/** Call this to set bPlayOnlyWhenVisible. */
+	UFUNCTION(BlueprintCallable, Category = "Media|MediaPlate")
+	void SetPlayOnlyWhenVisible(bool bInPlayOnlyWhenVisible);
+
 	/**
 	 * Adds our media texture to the media texture tracker.
 	 */
@@ -138,6 +142,10 @@ public:
 #endif
 
 private:
+	/** If true then only allow playback when the media plate is visible. */
+	UPROPERTY(EditAnywhere, BlueprintReadOnly, Category = "MediaPlate", meta = (AllowPrivateAccess = true))
+	bool bPlayOnlyWhenVisible = false;
+
 	UPROPERTY()
 	FVector2D MeshRange = FVector2D(360.0f, 180.0f);
 
@@ -158,6 +166,10 @@ private:
 	TSharedPtr<FMediaTextureTrackerObject, ESPMode::ThreadSafe> MediaTextureTrackerObject;
 	/** Our media clock sink. */
 	TSharedPtr<FMediaComponentClockSink, ESPMode::ThreadSafe> ClockSink;
+	/** Game time when we paused playback. */
+	float TimeWhenPlaybackPaused = 0.0f;
+	/** True if our media should be playing when visible. */
+	bool bWantsToPlayWhenVisible = false;
 
 	/**
 	 * Plays a media source.
@@ -172,4 +184,24 @@ private:
 	 */
 	void StopClockSink();
 
+	/**
+	 * Call this to see if this media plate is visible.
+	 */
+	bool IsVisible();
+
+	/**
+	 * Call this to resume playback when the media plate is visible.
+	 */
+	void ResumeWhenVisible();
+
+	/**
+	 * Returns the time to seek to when resuming playback.
+	 */
+	FTimespan GetResumeTime();
+
+	/**
+	 * Updates if we should tick or not based on current state.
+	 */
+	void UpdateTicking();
+
 };
