MediaPlate: Added letterboxes.

#rnx

diff --git a/Engine/Plugins/Media/MediaPlate/Source/MediaPlate/Private/MediaPlateComponent.cpp b/Engine/Plugins/Media/MediaPlate/Source/MediaPlate/Private/MediaPlateComponent.cpp
--- a/Engine/Plugins/Media/MediaPlate/Source/MediaPlate/Private/MediaPlateComponent.cpp
+++ b/Engine/Plugins/Media/MediaPlate/Source/MediaPlate/Private/MediaPlateComponent.cpp
@@ -7,6 +7,7 @@
 #include "IMediaClock.h"
 #include "IMediaClockSink.h"
 #include "IMediaModule.h"
+#include "Materials/Material.h"
 #include "MediaComponent.h"
 #include "MediaPlateModule.h"
 #include "MediaPlayer.h"
@@ -310,6 +311,24 @@ bool UMediaPlateComponent::PlayMediaSource(UMediaSource* InMediaSource)
 	return bIsPlaying;
 }
 
+
+float UMediaPlateComponent::GetAspectRatio()
+{
+	if (StaticMeshComponent != nullptr)
+	{
+		// Calculate aspect ratio from the scale.
+		FVector Scale = StaticMeshComponent->GetComponentScale();
+		float AspectRatio = 0.0f;
+		if (Scale.Z != 0.0f)
+		{
+			AspectRatio = Scale.Y / Scale.Z;
+		}
+		return AspectRatio;
+	}
+
+	return 0.0f;
+}
+
 void UMediaPlateComponent::SetAspectRatio(float AspectRatio)
 {
 	// Get the static mesh.
@@ -326,9 +345,26 @@ void UMediaPlateComponent::SetAspectRatio(float AspectRatio)
 		StaticMeshComponent->Modify();
 #endif // WITH_EDITOR
 		StaticMeshComponent->SetRelativeScale3D(Scale);
+
+		UpdateLetterboxes();
 	}
 }
 
+void UMediaPlateComponent::SetLetterboxAspectRatio(float AspectRatio)
+{
+	LetterboxAspectRatio = FMath::Max(0.0f, AspectRatio);
+	if (LetterboxAspectRatio == 0.0f)
+	{
+		RemoveLetterboxes();
+	}
+	else
+	{
+		AddLetterboxes();
+	}
+
+	UpdateLetterboxes();
+}
+
 void UMediaPlateComponent::TickOutput()
 {
 	if (MediaPlayer != nullptr)
@@ -415,6 +451,97 @@ void UMediaPlateComponent::UpdateTicking()
 	PrimaryComponentTick.SetTickFunctionEnable(bEnableTick);
 }
 
+
+void UMediaPlateComponent::UpdateLetterboxes()
+{
+	float AspectRatio = GetAspectRatio();
+	if ((AspectRatio <= LetterboxAspectRatio) || (LetterboxAspectRatio <= 0.0f))
+	{
+		for (TObjectPtr<UStaticMeshComponent> Letterbox : Letterboxes)
+		{
+			if (Letterbox != nullptr)
+			{
+				Letterbox->Modify();
+
+				Letterbox->SetVisibility(false);
+			}
+		}
+	}
+	else if (AspectRatio > 0.0f)
+	{
+		float DefaultHeight = 50.0f;
+		float VideoHeight = DefaultHeight / AspectRatio;
+		float MaxHeight = DefaultHeight / LetterboxAspectRatio;
+
+		float LetterboxHeight = (MaxHeight - VideoHeight) * 0.5f;
+		LetterboxHeight = FMath::Max(LetterboxHeight, 0.0f);
+		FVector Scale(1.0f, 1.0f, LetterboxHeight / DefaultHeight);
+
+		FVector Location(0.0f, 0.0f, VideoHeight + LetterboxHeight);
+
+		for (TObjectPtr<UStaticMeshComponent> Letterbox : Letterboxes)
+		{
+			if (Letterbox != nullptr)
+			{
+				Letterbox->Modify();
+				Letterbox->SetVisibility(true);
+				Letterbox->SetRelativeScale3D(Scale);
+				Letterbox->SetRelativeLocation(Location);
+				Location.Z = -Location.Z;
+			}
+		}
+	}
+}
+
+
+void UMediaPlateComponent::AddLetterboxes()
+{
+	if (Letterboxes.Num() == 0)
+	{
+		AActor* Owner = GetOwner();
+		if (Owner != nullptr)
+		{
+			UStaticMesh* Mesh = LoadObject<UStaticMesh>(NULL, TEXT("/MediaPlate/SM_MediaPlateScreen"), NULL, LOAD_None, NULL);
+			UMaterial* Material = LoadObject<UMaterial>(NULL, TEXT("/MediaPlate/M_MediaPlateLetterbox"), NULL, LOAD_None, NULL);
+			if ((Mesh != nullptr) && (Material != nullptr))
+			{
+				for (int32 Index = 0; Index < 2; ++Index)
+				{
+					UStaticMeshComponent* Letterbox = NewObject<UStaticMeshComponent>(Owner);
+					if (Letterbox != nullptr)
+					{
+						Letterboxes.Add(Letterbox);
+						Owner->Modify();
+						Owner->AddInstanceComponent(Letterbox);
+						Letterbox->OnComponentCreated();
+						Letterbox->AttachToComponent(Owner->GetRootComponent(),
+							FAttachmentTransformRules::KeepRelativeTransform);
+						Letterbox->RegisterComponent();
+						Letterbox->SetStaticMesh(Mesh);
+						Letterbox->SetMaterial(0, Material);
+						Letterbox->bCastStaticShadow = false;
+						Letterbox->bCastDynamicShadow = false;
+						Letterbox->SetVisibility(true);
+					}
+				}
+			}
+		}
+	}
+}
+
+void UMediaPlateComponent::RemoveLetterboxes()
+{
+	for (TObjectPtr<UStaticMeshComponent> Letterbox : Letterboxes)
+	{
+		if (Letterbox != nullptr)
+		{	
+			Letterbox->DestroyComponent();
+		}
+	}
+
+	Letterboxes.Empty();
+}
+
 #if WITH_EDITOR
 
 void UMediaPlateComponent::OnVisibleMipsTilesCalculationsChange()
diff --git a/Engine/Plugins/Media/MediaPlate/Source/MediaPlate/Public/MediaPlateComponent.h b/Engine/Plugins/Media/MediaPlate/Source/MediaPlate/Public/MediaPlateComponent.h
--- a/Engine/Plugins/Media/MediaPlate/Source/MediaPlate/Public/MediaPlateComponent.h
+++ b/Engine/Plugins/Media/MediaPlate/Source/MediaPlate/Public/MediaPlateComponent.h
@@ -89,6 +89,9 @@ public:
 	/** Holds the component for the mesh. */
 	UPROPERTY(EditAnywhere, AdvancedDisplay, Category = MediaPlate)
 	TObjectPtr<UStaticMeshComponent> StaticMeshComponent;
+	/** Holds the component for the mesh. */
+	UPROPERTY(EditAnywhere, AdvancedDisplay, Category = MediaPlate)
+	TArray<TObjectPtr<UStaticMeshComponent>> Letterboxes;
 
 	/** What media playlist to play. */
 	UPROPERTY(EditAnywhere, Category = "MediaPlate")
@@ -126,11 +129,28 @@ public:
 	 */
 	void UnregisterWithMediaTextureTracker();
 
+	/**
+	 * Call this to geet the aspect ratio of the mesh.
+	 */
+	float GetAspectRatio();
+
 	/**
 	 * Call this to set the aspect ratio of the mesh.
 	 */
 	void SetAspectRatio(float AspectRatio);
 
+	/**
+	 * Call this to get the aspect ratio of the screen.
+	 */
+	UFUNCTION(BlueprintCallable, Category = "Media|MediaPlate")
+	float GetLetterboxAspectRatio() { return LetterboxAspectRatio; }
+
+	/**
+	 * Call this to set the aspect ratio of the screen.
+	 */
+	UFUNCTION(BlueprintCallable, Category = "Media|MediaPlate")
+	void SetLetterboxAspectRatio(float AspectRatio);
+
 	/**
 	 * Called from the media clock.
 	 */
@@ -146,6 +166,11 @@ private:
 	UPROPERTY(EditAnywhere, BlueprintReadOnly, Category = "MediaPlate", meta = (AllowPrivateAccess = true))
 	bool bPlayOnlyWhenVisible = false;
 
+	/** If > 0, then this is the aspect ratio of our screen and 
+	 * letterboxes will be added if the media is smaller than the screen. */
+	UPROPERTY()
+	float LetterboxAspectRatio = 0.0f;
+
 	UPROPERTY()
 	FVector2D MeshRange = FVector2D(360.0f, 180.0f);
 
@@ -204,4 +229,19 @@ private:
 	 */
 	void UpdateTicking();
 
+	/**
+	 * Updates letterboxes based on the current state.
+	 */
+	void UpdateLetterboxes();
+
+	/**
+	 * Adds ability to have letterboxes.
+	 */
+	void AddLetterboxes();
+
+	/**
+	 * Removes ability to have letterboxes.
+	 */
+	void RemoveLetterboxes();
+
 };
diff --git a/Engine/Plugins/Media/MediaPlate/Source/MediaPlateEditor/Private/MediaPlateCustomization.cpp b/Engine/Plugins/Media/MediaPlate/Source/MediaPlateEditor/Private/MediaPlateCustomization.cpp
--- a/Engine/Plugins/Media/MediaPlate/Source/MediaPlateEditor/Private/MediaPlateCustomization.cpp
+++ b/Engine/Plugins/Media/MediaPlate/Source/MediaPlateEditor/Private/MediaPlateCustomization.cpp
@@ -417,6 +417,48 @@ void FMediaPlateCustomization::AddMeshCustomization(IDetailCategoryBuilder& Medi
 				]
 		];
 
+		// Add letterbox aspect ratio.
+		DetailGroup.AddWidgetRow()
+			.Visibility(MeshPlaneVisibility)
+			.NameContent()
+			[
+				SNew(STextBlock)
+					.Text(LOCTEXT("LetterboxAspectRatio", "Letterbox Aspect Ratio"))
+					.ToolTipText(LOCTEXT("LetterboxAspectRatio_ToolTip",
+						"Sets the aspect ratio of the whole screen.\n"
+						"If the screen is larger than the media then letterboxes will be added."))
+					.Font(IDetailLayoutBuilder::GetDetailFont())
+			]
+			.ValueContent()
+			[
+				SNew(SHorizontalBox)
+
+				// Presets button.
+				+ SHorizontalBox::Slot()
+				.AutoWidth()
+				[
+					SNew(SComboButton)
+						.OnGetMenuContent(this, &FMediaPlateCustomization::OnGetLetterboxAspectRatios)
+						.ContentPadding(2)
+						.ButtonContent()
+						[
+							SNew(STextBlock)
+								.ToolTipText(LOCTEXT("Presets_ToolTip", "Select one of the presets for the aspect ratio."))
+								.Text(LOCTEXT("Presets", "Presets"))
+						]
+			]
+
+		// Numeric entry box.
+		+ SHorizontalBox::Slot()
+			.AutoWidth()
+			[
+				SNew(SNumericEntryBox<float>)
+					.Value(this, &FMediaPlateCustomization::GetLetterboxAspectRatio)
+					.MinValue(0.0f)
+					.OnValueChanged(this, &FMediaPlateCustomization::SetLetterboxAspectRatio)
+			]
+		];
+
 	// Add auto aspect ratio.
 	DetailGroup.AddWidgetRow()
 		.Visibility(MeshPlaneVisibility)
@@ -502,9 +544,15 @@ void FMediaPlateCustomization::SetMeshMode(EMediaTextureVisibleMipsTiles InMode)
 				{
 					MeshCustomization.SetPlaneMesh(MediaPlate);
 				}
-				else if (MeshMode == EMediaTextureVisibleMipsTiles::Sphere)
+				else
 				{
-					SetSphereMesh(MediaPlate);
+					// Letterboxes are only for planes.
+					SetLetterboxAspectRatio(0.0f);
+
+					if (MeshMode == EMediaTextureVisibleMipsTiles::Sphere)
+					{
+						SetSphereMesh(MediaPlate);
+					}
 				}
 			}
 		}
@@ -560,23 +608,42 @@ void FMediaPlateCustomization::SetIsAspectRatioAuto(ECheckBoxState State)
 
 TSharedRef<SWidget> FMediaPlateCustomization::OnGetAspectRatios()
 {
-	FMenuBuilder MenuBuilder(true, NULL);
+	FUIAction Actions[4];
+	Actions[0] = FUIAction(FExecuteAction::CreateSP(this, &FMediaPlateCustomization::SetAspectRatio, 16.0f / 9.0f));
+	Actions[1] = FUIAction(FExecuteAction::CreateSP(this, &FMediaPlateCustomization::SetAspectRatio, 16.0f / 10.0f));
+	Actions[2] = FUIAction(FExecuteAction::CreateSP(this, &FMediaPlateCustomization::SetAspectRatio, 4.0f / 3.0f));
+	Actions[3] = FUIAction(FExecuteAction::CreateSP(this, &FMediaPlateCustomization::SetAspectRatio, 1.0f));
 
-	FUIAction Set16x9Action(FExecuteAction::CreateSP(this, &FMediaPlateCustomization::SetAspectRatio, 16.0f / 9.0f));
-	MenuBuilder.AddMenuEntry(LOCTEXT("16x9", "16x9"), FText(), FSlateIcon(), Set16x9Action);
+	FMenuBuilder MenuBuilder(true, NULL);
+	AddAspectRatiosToMenuBuilder(MenuBuilder, Actions);
 
-	FUIAction Set16x10Action(FExecuteAction::CreateSP(this, &FMediaPlateCustomization::SetAspectRatio, 16.0f / 10.0f));
-	MenuBuilder.AddMenuEntry(LOCTEXT("16x10", "16x10"), FText(), FSlateIcon(), Set16x10Action);
+	return MenuBuilder.MakeWidget();
+}
 
-	FUIAction Set4x3Action(FExecuteAction::CreateSP(this, &FMediaPlateCustomization::SetAspectRatio, 4.0f / 3.0f));
-	MenuBuilder.AddMenuEntry(LOCTEXT("4x3", "4x3"), FText(), FSlateIcon(), Set4x3Action);
+TSharedRef<SWidget> FMediaPlateCustomization::OnGetLetterboxAspectRatios()
+{
+	FUIAction Actions[5];
+	Actions[0] = FUIAction(FExecuteAction::CreateSP(this, &FMediaPlateCustomization::SetLetterboxAspectRatio, 16.0f / 9.0f));
+	Actions[1] = FUIAction(FExecuteAction::CreateSP(this, &FMediaPlateCustomization::SetLetterboxAspectRatio, 16.0f / 10.0f));
+	Actions[2] = FUIAction(FExecuteAction::CreateSP(this, &FMediaPlateCustomization::SetLetterboxAspectRatio, 4.0f / 3.0f));
+	Actions[3] = FUIAction(FExecuteAction::CreateSP(this, &FMediaPlateCustomization::SetLetterboxAspectRatio, 1.0f));
+	Actions[4] = FUIAction(FExecuteAction::CreateSP(this, &FMediaPlateCustomization::SetLetterboxAspectRatio, 0.0f));
 
-	FUIAction Set1x1Action(FExecuteAction::CreateSP(this, &FMediaPlateCustomization::SetAspectRatio, 1.0f));
-	MenuBuilder.AddMenuEntry(LOCTEXT("1x1", "1x1"), FText(), FSlateIcon(), Set1x1Action);
+	FMenuBuilder MenuBuilder(true, NULL);
+	AddAspectRatiosToMenuBuilder(MenuBuilder, Actions);
+	MenuBuilder.AddMenuEntry(LOCTEXT("Disable", "Disable"), FText(), FSlateIcon(), Actions[4]);
 
 	return MenuBuilder.MakeWidget();
 }
 
+void FMediaPlateCustomization::AddAspectRatiosToMenuBuilder(FMenuBuilder& MenuBuilder, FUIAction Actions[])
+{
+	MenuBuilder.AddMenuEntry(LOCTEXT("16x9", "16x9"), FText(), FSlateIcon(), Actions[0]);
+	MenuBuilder.AddMenuEntry(LOCTEXT("16x10", "16x10"), FText(), FSlateIcon(), Actions[1]);
+	MenuBuilder.AddMenuEntry(LOCTEXT("4x3", "4x3"), FText(), FSlateIcon(), Actions[2]);
+	MenuBuilder.AddMenuEntry(LOCTEXT("1x1", "1x1"), FText(), FSlateIcon(), Actions[3]);
+}
+
 void FMediaPlateCustomization::SetAspectRatio(float AspectRatio)
 {
 	// Loop through all our objects.
@@ -605,25 +672,40 @@ TOptional<float> FMediaPlateCustomization::GetAspectRatio() const
 		UMediaPlateComponent* MediaPlate = MediaPlatePtr.Get();
 		if (MediaPlate != nullptr)
 		{
-			AActor* Owner = MediaPlate->GetOwner();
-			AMediaPlate* MediaPlateActor = Cast<AMediaPlate>(Owner);
-			if (MediaPlateActor != nullptr)
-			{
-				// Get the static mesh.
-				UStaticMeshComponent* StaticMeshComponent = MediaPlateActor->StaticMeshComponent;
-				if (StaticMeshComponent != nullptr)
-				{
-					// Calculate aspect ratio from the scale.
-					FVector Scale = StaticMeshComponent->GetComponentScale();
-					float AspectRatio = 0.0f;
-					if (Scale.Z != 0.0f)
-					{
-						AspectRatio = Scale.Y / Scale.Z;
-					}
-					return AspectRatio;
-				}
-			}
-			break;
+			return MediaPlate->GetAspectRatio();
+		}
+	}
+
+	return TOptional<float>();
+}
+
+void FMediaPlateCustomization::SetLetterboxAspectRatio(float AspectRatio)
+{
+	for (const TWeakObjectPtr<UMediaPlateComponent>& MediaPlatePtr : MediaPlatesList)
+	{
+		UMediaPlateComponent* MediaPlate = MediaPlatePtr.Get();
+		if (MediaPlate != nullptr)
+		{
+			MediaPlate->SetLetterboxAspectRatio(AspectRatio);
+		}
+	}
+
+	// Invalidate the viewport so we can see the mesh change.
+	if (GCurrentLevelEditingViewportClient != nullptr)
+	{
+		GCurrentLevelEditingViewportClient->Invalidate();
+	}
+}
+
+
+TOptional<float> FMediaPlateCustomization::GetLetterboxAspectRatio() const
+{
+	for (const TWeakObjectPtr<UMediaPlateComponent>& MediaPlatePtr : MediaPlatesList)
+	{
+		UMediaPlateComponent* MediaPlate = MediaPlatePtr.Get();
+		if (MediaPlate != nullptr)
+		{
+			return MediaPlate->GetLetterboxAspectRatio();
 		}
 	}
 
diff --git a/Engine/Plugins/Media/MediaPlate/Source/MediaPlateEditor/Private/MediaPlateCustomization.h b/Engine/Plugins/Media/MediaPlate/Source/MediaPlateEditor/Private/MediaPlateCustomization.h
--- a/Engine/Plugins/Media/MediaPlate/Source/MediaPlateEditor/Private/MediaPlateCustomization.h
+++ b/Engine/Plugins/Media/MediaPlate/Source/MediaPlateEditor/Private/MediaPlateCustomization.h
@@ -81,6 +81,16 @@ private:
 	 */
 	TSharedRef<SWidget> OnGetAspectRatios();
 
+	/**
+	 * Returns menu options for all aspect ratio presets.
+	 */
+	TSharedRef<SWidget> OnGetLetterboxAspectRatios();
+
+	/**
+	 * Returns menu options for all aspect ratio presets.
+	 */
+	void AddAspectRatiosToMenuBuilder(FMenuBuilder& MenuBuilder, FUIAction Actions[]);
+
 	/**
 	 * Call this to see if auto aspect ratio is enabled.
 	 */
@@ -101,6 +111,16 @@ private:
 	 */
 	TOptional<float> GetAspectRatio() const;
 
+	/**
+	 * Call this to set the aspect ratio.
+	 */
+	void SetLetterboxAspectRatio(float AspectRatio);
+
+	/**
+	 * Call this to get the aspect ratio.
+	 */
+	TOptional<float> GetLetterboxAspectRatio() const;
+
 	/**
 	 * Call this to set the horizontal range of the mesh.
 	 */
