MediaPlate: Fixed aspect ratio looking at wrong scale.
#rnx

diff --git a/Engine/Plugins/Media/MediaPlate/Source/MediaPlate/Private/MediaPlateComponent.cpp b/Engine/Plugins/Media/MediaPlate/Source/MediaPlate/Private/MediaPlateComponent.cpp
--- a/Engine/Plugins/Media/MediaPlate/Source/MediaPlate/Private/MediaPlateComponent.cpp
+++ b/Engine/Plugins/Media/MediaPlate/Source/MediaPlate/Private/MediaPlateComponent.cpp
@@ -317,7 +317,7 @@ float UMediaPlateComponent::GetAspectRatio()
 	if (StaticMeshComponent != nullptr)
 	{
 		// Calculate aspect ratio from the scale.
-		FVector Scale = StaticMeshComponent->GetComponentScale();
+		FVector Scale = StaticMeshComponent->GetRelativeScale3D();
 		float AspectRatio = 0.0f;
 		if (Scale.Z != 0.0f)
 		{
