MediaPlate: Forward/reverse buttons are disabled if they are not supported.
#rnx

diff --git a/Engine/Plugins/Media/MediaPlate/Source/MediaPlateEditor/Private/MediaPlateCustomization.cpp b/Engine/Plugins/Media/MediaPlate/Source/MediaPlateEditor/Private/MediaPlateCustomization.cpp
--- a/Engine/Plugins/Media/MediaPlate/Source/MediaPlateEditor/Private/MediaPlateCustomization.cpp
+++ b/Engine/Plugins/Media/MediaPlate/Source/MediaPlateEditor/Private/MediaPlateCustomization.cpp
@@ -162,6 +162,21 @@ void FMediaPlateCustomization::CustomizeDetails(IDetailLayoutBuilder& DetailBuil
 				[
 					SNew(SButton)
 						.VAlign(VAlign_Center)
+						.IsEnabled_Lambda([this] {
+							for (TWeakObjectPtr<UMediaPlateComponent>& MediaPlatePtr : MediaPlatesList)
+							{
+								UMediaPlateComponent* MediaPlate = MediaPlatePtr.Get();
+								if (MediaPlate != nullptr)
+								{
+									UMediaPlayer* MediaPlayer = MediaPlate->GetMediaPlayer();
+									if (MediaPlayer != nullptr)
+									{
+										return MediaPlayer->IsReady() && MediaPlayer->SupportsRate(GetReverseRate(MediaPlayer), false);
+									}
+								}
+							}
+							return false;
+						})
 						.OnClicked_Lambda([this]() -> FReply
 						{
 							for (TWeakObjectPtr<UMediaPlateComponent>& MediaPlatePtr : MediaPlatesList)
@@ -267,6 +282,21 @@ void FMediaPlateCustomization::CustomizeDetails(IDetailLayoutBuilder& DetailBuil
 				[
 					SNew(SButton)
 					.VAlign(VAlign_Center)
+					.IsEnabled_Lambda([this] {
+						for (TWeakObjectPtr<UMediaPlateComponent>& MediaPlatePtr : MediaPlatesList)
+						{
+							UMediaPlateComponent* MediaPlate = MediaPlatePtr.Get();
+							if (MediaPlate != nullptr)
+							{
+								UMediaPlayer* MediaPlayer = MediaPlate->GetMediaPlayer();
+								if (MediaPlayer != nullptr)
+								{
+									return MediaPlayer->IsReady() && MediaPlayer->SupportsRate(GetForwardRate(MediaPlayer), false);
+								}
+							}
+						}
+						return false;
+					})
 					.OnClicked_Lambda([this]() -> FReply
 					{
 						for (TWeakObjectPtr<UMediaPlateComponent>& MediaPlatePtr : MediaPlatesList)
