MediaPlate: New material overrides are now correctly reconnected with the media texture.

This is now also compatible with the VPUtilities MID functionality, and works with drag-drop in-editor viewport.

#rnx

diff --git a/Engine/Plugins/Media/MediaPlate/Source/MediaPlate/Private/MediaPlate.cpp b/Engine/Plugins/Media/MediaPlate/Source/MediaPlate/Private/MediaPlate.cpp
--- a/Engine/Plugins/Media/MediaPlate/Source/MediaPlate/Private/MediaPlate.cpp
+++ b/Engine/Plugins/Media/MediaPlate/Source/MediaPlate/Private/MediaPlate.cpp
@@ -5,6 +5,7 @@
 #include "Components/StaticMeshComponent.h"
 #include "Engine/StaticMesh.h"
 #include "Materials/MaterialInstanceConstant.h"
+#include "Materials/MaterialInstanceDynamic.h"
 #include "MediaPlateComponent.h"
 #include "MediaPlateModule.h"
 #include "MediaTexture.h"
@@ -86,33 +87,58 @@ void AMediaPlate::BeginDestroy()
 
 void AMediaPlate::UseDefaultMaterial()
 {
-	// Get material.
-	UMaterial* Material = LoadObject<UMaterial>(NULL, TEXT("/MediaPlate/M_MediaPlate"), NULL, LOAD_None, NULL);
-	if (Material != nullptr)
+	UMaterial* DefaultMaterial = LoadObject<UMaterial>(NULL, TEXT("/MediaPlate/M_MediaPlate"), NULL, LOAD_None, NULL);
+	
+	ApplyMaterial(DefaultMaterial);
+}
+
+void AMediaPlate::ApplyMaterial(UMaterialInterface* Material)
+{
+	if (Material != nullptr && StaticMeshComponent != nullptr)
 	{
-		// Change M_ to MI_ in material name and then generate a unique one.
-		FString MaterialName = Material->GetName();
-		if (MaterialName.StartsWith(TEXT("M_")))
+		UMaterialInstanceDynamic* MID = Cast<UMaterialInstanceDynamic>(Material);
+		UMaterialInterface* Result = nullptr;
+
+		if (MID != nullptr)
 		{
-			MaterialName.InsertAt(1, TEXT("I"));
+			MID->SetTextureParameterValue(MediaTextureName, MediaPlateComponent->GetMediaTexture());
+
+			Result = MID;
+		}
+		else
+		{
+			// Change M_ to MI_ in material name and then generate a unique one.
+			FString MaterialName = Material->GetName();
+			if (MaterialName.StartsWith(TEXT("M_")))
+			{
+				MaterialName.InsertAt(1, TEXT("I"));
+			}
+			FName MaterialUniqueName = MakeUniqueObjectName(StaticMeshComponent, UMaterialInstanceConstant::StaticClass(),
+				FName(*MaterialName));
+
+			// Create instance.
+			UMaterialInstanceConstant* MaterialInstance =
+				NewObject<UMaterialInstanceConstant>(StaticMeshComponent, MaterialUniqueName);
+			MaterialInstance->SetParentEditorOnly(Material);
+			MaterialInstance->CopyMaterialUniformParametersEditorOnly(Material);
+			MaterialInstance->SetTextureParameterValueEditorOnly(
+				FMaterialParameterInfo(MediaTextureName),
+				MediaPlateComponent->GetMediaTexture());
+			MaterialInstance->PostEditChange();
+
+			// We force call post-load to indirectly call UpdateParameters() (for integration with VPUtilities plugin).
+			MaterialInstance->PostLoad();
+
+			Result = MaterialInstance;
 		}
-		FName MaterialUniqueName = MakeUniqueObjectName(this, UMaterialInstanceConstant::StaticClass(),
-			FName(*MaterialName));
-
-		// Create instance.
-		UMaterialInstanceConstant* MaterialInstance =
-			NewObject<UMaterialInstanceConstant>(this, MaterialUniqueName);
-		MaterialInstance->SetParentEditorOnly(Material);
-		MaterialInstance->SetTextureParameterValueEditorOnly(
-			FMaterialParameterInfo(MediaTextureName),
-			MediaPlateComponent->GetMediaTexture());
-		MaterialInstance->PostEditChange();
 
 		// Update static mesh.
-		if (StaticMeshComponent != nullptr)
+		if (Result != nullptr)
 		{
 			StaticMeshComponent->Modify();
-			StaticMeshComponent->SetMaterial(0, MaterialInstance);
+			StaticMeshComponent->SetMaterial(0, Result);
+
+			LastMaterial = Result;
 		}
 	}
 }
diff --git a/Engine/Plugins/Media/MediaPlate/Source/MediaPlate/Public/MediaPlate.h b/Engine/Plugins/Media/MediaPlate/Source/MediaPlate/Public/MediaPlate.h
--- a/Engine/Plugins/Media/MediaPlate/Source/MediaPlate/Public/MediaPlate.h
+++ b/Engine/Plugins/Media/MediaPlate/Source/MediaPlate/Public/MediaPlate.h
@@ -36,7 +36,8 @@ public:
 	 * Call this to change the static mesh to use the default media plate material.
 	 */
 	void UseDefaultMaterial();
-
+	void ApplyMaterial(UMaterialInterface* InMaterial);
+	UMaterialInterface* GetLastMaterial() const { return LastMaterial; }
 #endif
 
 private:
@@ -44,4 +45,8 @@ private:
 	static FLazyName MediaPlateComponentName;
 	/** Name for the media texture parameter in the material. */
 	static FLazyName MediaTextureName;
+
+#if WITH_EDITOR
+	UMaterialInterface* LastMaterial = nullptr;
+#endif
 };
diff --git a/Engine/Plugins/Media/MediaPlate/Source/MediaPlateEditor/Private/MediaPlateCustomization.cpp b/Engine/Plugins/Media/MediaPlate/Source/MediaPlateEditor/Private/MediaPlateCustomization.cpp
--- a/Engine/Plugins/Media/MediaPlate/Source/MediaPlateEditor/Private/MediaPlateCustomization.cpp
+++ b/Engine/Plugins/Media/MediaPlate/Source/MediaPlateEditor/Private/MediaPlateCustomization.cpp
@@ -30,6 +30,45 @@
 /* IDetailCustomization interface
  *****************************************************************************/
 
+
+FMediaPlateCustomization::FMediaPlateCustomization()
+{
+	PropertyChangeDelegate = FCoreUObjectDelegates::OnObjectPropertyChanged.AddRaw(this, &FMediaPlateCustomization::OnObjectPropertyChanged);
+}
+
+FMediaPlateCustomization::~FMediaPlateCustomization()
+{
+	FCoreUObjectDelegates::OnObjectPropertyChanged.RemoveAll(this);
+}
+
+void FMediaPlateCustomization::OnObjectPropertyChanged(UObject* InObject, FPropertyChangedEvent& InEvent)
+{
+	//Note: Cannot check for GET_MEMBER_NAME_CHECKED(UMeshComponent, OverrideMaterials) because the event is null on drag-and-drop.
+
+	for (TWeakObjectPtr<UMediaPlateComponent>& MediaPlateComponent : MediaPlatesList)
+	{
+		if (MediaPlateComponent.IsValid())
+		{
+			if (InObject->GetOuter() == MediaPlateComponent->GetOuter())
+			{
+				AMediaPlate* MediaPlate = Cast<AMediaPlate>(InObject->GetOuter());
+
+				if (MediaPlate != nullptr && MediaPlate->StaticMeshComponent != nullptr)
+				{
+					UMaterialInterface* Material = MediaPlate->StaticMeshComponent->GetMaterial(0);
+
+					if (Material != MediaPlate->GetLastMaterial())
+					{
+						MediaPlate->ApplyMaterial(Material);
+					}
+				}
+
+				break;
+			}
+		}
+	}
+}
+
 void FMediaPlateCustomization::CustomizeDetails(IDetailLayoutBuilder& DetailBuilder)
 {
 	// Is this the media plate editor window?
diff --git a/Engine/Plugins/Media/MediaPlate/Source/MediaPlateEditor/Private/MediaPlateCustomization.h b/Engine/Plugins/Media/MediaPlate/Source/MediaPlateEditor/Private/MediaPlateCustomization.h
--- a/Engine/Plugins/Media/MediaPlate/Source/MediaPlateEditor/Private/MediaPlateCustomization.h
+++ b/Engine/Plugins/Media/MediaPlate/Source/MediaPlateEditor/Private/MediaPlateCustomization.h
@@ -24,6 +24,9 @@ class FMediaPlateCustomization
 {
 public:
 
+	FMediaPlateCustomization();
+	~FMediaPlateCustomization();
+
 	//~ IDetailCustomization interface
 
 	virtual void CustomizeDetails(IDetailLayoutBuilder& DetailBuilder) override;
@@ -40,6 +43,8 @@ public:
 
 private:
 
+	/** Property change delegate used for static mesh material changes. */
+	FDelegateHandle PropertyChangeDelegate;
 	/** List of the media plates we are editing. */
 	TArray<TWeakObjectPtr<UMediaPlateComponent>> MediaPlatesList;
 	/** Stores the current value of the MediaPath property. */
@@ -51,6 +56,11 @@ private:
 	/** Handles mesh stuff. */
 	FMediaPlateCustomizationMesh MeshCustomization;
 
+	/**
+	 * Object property change callback used for static mesh material changes.
+	 */
+	void OnObjectPropertyChanged(UObject* InObject, FPropertyChangedEvent& InEvent);
+
 	/**
 	 * Adds widgets for editing the mesh.
 	 */
