MediaPlate: Media track will be set to the length of the media when adding a media plate to sequencer.

#rnx

diff --git a/Engine/Plugins/Media/MediaPlate/Source/MediaPlateEditor/Private/Sequencer/MediaPlateTrackEditor.cpp b/Engine/Plugins/Media/MediaPlate/Source/MediaPlateEditor/Private/Sequencer/MediaPlateTrackEditor.cpp
--- a/Engine/Plugins/Media/MediaPlate/Source/MediaPlateEditor/Private/Sequencer/MediaPlateTrackEditor.cpp
+++ b/Engine/Plugins/Media/MediaPlate/Source/MediaPlateEditor/Private/Sequencer/MediaPlateTrackEditor.cpp
@@ -7,6 +7,7 @@
 #include "MediaTexture.h"
 #include "MediaPlate.h"
 #include "MediaPlateComponent.h"
+#include "MediaPlayer.h"
 #include "MediaPlaylist.h"
 #include "MovieSceneMediaSection.h"
 #include "MovieSceneMediaTrack.h"
@@ -26,6 +27,7 @@ TArray<FAnimatedPropertyKey, TInlineAllocator<1>> FMediaPlateTrackEditor::GetAni
 
 FMediaPlateTrackEditor::FMediaPlateTrackEditor(TSharedRef<ISequencer> InSequencer)
 	: FMovieSceneTrackEditor(InSequencer)
+	, bGetDurationDelay(false)
 {
 	OnActorAddedToSequencerHandle = InSequencer->OnActorAddedToSequencer().AddRaw(this, &FMediaPlateTrackEditor::HandleActorAdded);
 }
@@ -73,6 +75,42 @@ bool FMediaPlateTrackEditor::SupportsType(TSubclassOf<UMovieSceneTrack> TrackCla
 	return false;
 }
 
+void FMediaPlateTrackEditor::Tick(float DeltaTime)
+{
+	// Do we have any new sections that need durations?
+	if (NewSections.Num() > 0)
+	{
+		// Can we get the duration?
+		if (bGetDurationDelay)
+		{
+			bGetDurationDelay = false;
+		}
+		else
+		{
+			// Loop over all new sections.
+			for (int32 Index = 0; Index < NewSections.Num();)
+			{
+				if (NewSections[Index].Key.IsValid())
+				{
+					// Try and get the duration.
+					if (GetDuration(NewSections[Index].Key, NewSections[Index].Value))
+					{
+						NewSections.RemoveAtSwap(Index);
+					}
+					else
+					{
+						++Index;
+					}
+				}
+				else
+				{
+					NewSections.RemoveAtSwap(Index);
+				}
+			}
+		}
+	}
+}
+
 void FMediaPlateTrackEditor::HandleAddMediaTrackToObjectBindingMenuEntryExecute(TArray<FGuid> InObjectBindingIDs)
 {
 	UMovieScene* FocusedMovieScene = GetFocusedMovieScene();
@@ -143,12 +181,17 @@ void FMediaPlateTrackEditor::AddTrackForComponent(UMediaPlateComponent* Componen
 				if (MediaSource != nullptr)
 				{
 					UMovieSceneSection* Section = MediaTrack->AddNewMediaSource(*MediaSource, FFrameNumber(0));
-					
-					// Copy cache settings from media plate.
-					UMovieSceneMediaSection* MediaSection = Cast<UMovieSceneMediaSection>(Section);
-					if (MediaSection != nullptr)
+					if (Section != nullptr)
 					{
-						MediaSection->CacheSettings = Component->CacheSettings;
+						// Start process to get the duration.
+						StartGetDuration(MediaSource, Section);
+
+						// Copy cache settings from media plate.
+						UMovieSceneMediaSection* MediaSection = Cast<UMovieSceneMediaSection>(Section);
+						if (MediaSection != nullptr)
+						{
+							MediaSection->CacheSettings = Component->CacheSettings;
+						}
 					}
 				}
 			}
@@ -189,5 +232,52 @@ void FMediaPlateTrackEditor::ImportObjectBinding(const TArray<FGuid> ObjectBindi
 	}
 }
 
+void FMediaPlateTrackEditor::StartGetDuration(UMediaSource* MediaSource, UMovieSceneSection* Section)
+{
+	// Create media player.
+	TStrongObjectPtr<UMediaPlayer> MediaPlayer = TStrongObjectPtr<UMediaPlayer>(
+		NewObject<UMediaPlayer>(GetTransientPackage(),
+			MakeUniqueObjectName(GetTransientPackage(),
+				UMediaPlayer::StaticClass())));
+
+	// Open the media.
+	MediaPlayer->PlayOnOpen = false;
+	if (MediaPlayer->OpenSource(MediaSource))
+	{
+		NewSections.Emplace(MediaPlayer, Section);
+	}
+
+	// Some players like Electra report that they are closed at this point, so wait a frame.
+	bGetDurationDelay = true;
+}
+
+bool FMediaPlateTrackEditor::GetDuration(
+	TStrongObjectPtr<UMediaPlayer>& MediaPlayer, TWeakObjectPtr<UMovieSceneSection>& NewSection)
+{
+	bool bIsDone = false;
+	
+	// Check everything is ok.
+	if ((MediaPlayer.IsValid() == false) || (MediaPlayer->HasError()) || (MediaPlayer->IsClosed()) ||
+		(NewSection.IsValid() == false))
+	{
+		bIsDone = true;
+	}
+	else
+	{
+		// Get the duration.
+		FTimespan Duration = MediaPlayer->GetDuration();
+		if (Duration != 0)
+		{
+			// Once it is non zero, then set the length of the section.
+			FFrameRate TickResolution = NewSection->GetTypedOuter<UMovieScene>()->GetTickResolution();
+			FFrameNumber StartFrame = NewSection->GetInclusiveStartFrame();
+			FFrameNumber EndFrame = StartFrame + (Duration.GetTotalSeconds() * TickResolution).FrameNumber;
+			NewSection->SetEndFrame(TRangeBound<FFrameNumber>::Exclusive(EndFrame));
+			bIsDone = true;
+		}
+	}
+
+	return bIsDone;
+}
 
 #undef LOCTEXT_NAMESPACE
diff --git a/Engine/Plugins/Media/MediaPlate/Source/MediaPlateEditor/Private/Sequencer/MediaPlateTrackEditor.h b/Engine/Plugins/Media/MediaPlate/Source/MediaPlateEditor/Private/Sequencer/MediaPlateTrackEditor.h
--- a/Engine/Plugins/Media/MediaPlate/Source/MediaPlateEditor/Private/Sequencer/MediaPlateTrackEditor.h
+++ b/Engine/Plugins/Media/MediaPlate/Source/MediaPlateEditor/Private/Sequencer/MediaPlateTrackEditor.h
@@ -3,12 +3,15 @@
 #pragma once
 
 #include "AnimatedPropertyKey.h"
+#include "Containers/Map.h"
 #include "MovieSceneTrackEditor.h"
 #include "Templates/SharedPointer.h"
 
 class AActor;
 class ISequencer;
 class UMediaPlateComponent;
+class UMediaPlayer;
+class UMediaSource;
 
 /**
  * Track editor for media plate.
@@ -54,6 +57,7 @@ public:
 	virtual void BuildObjectBindingContextMenu(FMenuBuilder& MenuBuilder, const TArray<FGuid>& ObjectBindings, const UClass* ObjectClass) override;
 	virtual void BuildObjectBindingTrackMenu(FMenuBuilder& MenuBuilder, const TArray<FGuid>& ObjectBindings, const UClass* ObjectClass) override;
 	virtual bool SupportsType(TSubclassOf<UMovieSceneTrack> TrackClass) const override;
+	virtual void Tick(float DeltaTime) override;
 	//~ ISequencerTrackEditor interface
 	virtual void OnRelease() override;
 
@@ -76,6 +80,31 @@ private:
 	 */
 	void ImportObjectBinding(const TArray<FGuid> ObjectBindings);
 
+	/**
+	 * Starts the process to get the duration of the media.
+	 * It might take a frame or more.
+	 * 
+	 * @param MediaSource		Media to inspect.
+	 * @param Section			Will set this sequencer section to the length of the media.
+	*/
+	void StartGetDuration(UMediaSource* MediaSource, UMovieSceneSection* Section);
+
+	/**
+	 * Call this after StartGetDuration to try and get the duration of the media.
+	 *
+	 * @param MediaPlayer		Player that is opening the media.
+	 * @param NewSection		Movie section this is for.
+	 * @return True if it is done and the player can be removed.
+	 */
+	bool GetDuration(TStrongObjectPtr<UMediaPlayer>& MediaPlayer,
+		TWeakObjectPtr<UMovieSceneSection>& NewSection);
+
 	/** Handle to our delegate. */
 	FDelegateHandle OnActorAddedToSequencerHandle;
+
+	/** List of players that are we are trying to get durations from for the corresponding sections. */
+	TArray<TPair<TStrongObjectPtr<UMediaPlayer>, TWeakObjectPtr<UMovieSceneSection>>> NewSections;
+
+	/** If true then don't check for the duration this frame. */
+	bool bGetDurationDelay;
 };
