MediaPlate: Exposed start time to blueprints.
#rnx

diff --git a/Engine/Plugins/Media/MediaPlate/Source/MediaPlate/Public/MediaPlateComponent.h b/Engine/Plugins/Media/MediaPlate/Source/MediaPlate/Public/MediaPlateComponent.h
--- a/Engine/Plugins/Media/MediaPlate/Source/MediaPlate/Public/MediaPlateComponent.h
+++ b/Engine/Plugins/Media/MediaPlate/Source/MediaPlate/Public/MediaPlateComponent.h
@@ -75,7 +75,7 @@ public:
 	bool bEnableAudio = false;
 
 	/** What time to start playing from (in seconds). */
-	UPROPERTY(EditAnywhere, Category = "MediaPlate", meta = (ClampMin = "0.0"))
+	UPROPERTY(EditAnywhere, BlueprintReadWrite, Category = "MediaPlate", meta = (ClampMin = "0.0"))
 	float StartTime = 0.0f;
 
 	/** If true then set the aspect ratio automatically based on the media. */
