MediaPlate: Added media plates as options when adding media tracks.
#rnx

diff --git a/Engine/Plugins/Media/MediaCompositing/Source/MediaCompositingEditor/Private/Sequencer/MediaTrackEditor.cpp b/Engine/Plugins/Media/MediaCompositing/Source/MediaCompositingEditor/Private/Sequencer/MediaTrackEditor.cpp
--- a/Engine/Plugins/Media/MediaCompositing/Source/MediaCompositingEditor/Private/Sequencer/MediaTrackEditor.cpp
+++ b/Engine/Plugins/Media/MediaCompositing/Source/MediaCompositingEditor/Private/Sequencer/MediaTrackEditor.cpp
@@ -27,6 +27,7 @@
 
 #define LOCTEXT_NAMESPACE "FMediaTrackEditor"
 
+FOnBuildOutlinerEditWidget FMediaTrackEditor::OnBuildOutlinerEditWidget;
 
 /* FMediaTrackEditor static functions
  *****************************************************************************/
@@ -111,7 +112,20 @@ TSharedPtr<SWidget> FMediaTrackEditor::BuildOutlinerEditWidget(const FGuid& Obje
 				ContentBrowserModule.Get().CreateAssetPicker(AssetPickerConfig)
 			];
 
-		return Picker;
+		// Create menu.
+		FMenuBuilder MenuBuilder(true, NULL);
+		MenuBuilder.AddSubMenu(
+			LOCTEXT("MediaSourceLabel", "Media Source"),
+			LOCTEXT("MediaSourceTooltip", "Add media from a media source."),
+			FNewMenuDelegate::CreateLambda([Picker](FMenuBuilder& InMenuBuilder)
+		{
+			InMenuBuilder.AddWidget(Picker, FText::GetEmpty(), true);
+		}));
+
+		// Allow customizations from other sources.
+		OnBuildOutlinerEditWidget.Broadcast(MenuBuilder);
+
+		return MenuBuilder.MakeWidget();
 	};
 
 	return SNew(SHorizontalBox)
diff --git a/Engine/Plugins/Media/MediaCompositing/Source/MediaCompositingEditor/Public/Sequencer/MediaTrackEditor.h b/Engine/Plugins/Media/MediaCompositing/Source/MediaCompositingEditor/Public/Sequencer/MediaTrackEditor.h
--- a/Engine/Plugins/Media/MediaCompositing/Source/MediaCompositingEditor/Public/Sequencer/MediaTrackEditor.h
+++ b/Engine/Plugins/Media/MediaCompositing/Source/MediaCompositingEditor/Public/Sequencer/MediaTrackEditor.h
@@ -12,6 +12,7 @@ class FTrackEditorThumbnailPool;
 class UMediaSource;
 class UMovieSceneMediaTrack;
 
+DECLARE_EVENT_OneParam(FMediaTrackEditor, FOnBuildOutlinerEditWidget, FMenuBuilder&);
 
 /**
  * Track editor that understands how to animate MediaPlayer properties on objects
@@ -39,6 +40,12 @@ public:
 	 */
 	static TArray<FAnimatedPropertyKey, TInlineAllocator<1>> GetAnimatedPropertyTypes();
 
+	/**
+	 * Event for when we build the widget for adding to the track.
+	 * Hook into this if you want to add custom options.
+	 */
+	static FOnBuildOutlinerEditWidget OnBuildOutlinerEditWidget;
+
 public:
 
 	/**
diff --git a/Engine/Plugins/Media/MediaPlate/Source/MediaPlateEditor/MediaPlateEditor.Build.cs b/Engine/Plugins/Media/MediaPlate/Source/MediaPlateEditor/MediaPlateEditor.Build.cs
--- a/Engine/Plugins/Media/MediaPlate/Source/MediaPlateEditor/MediaPlateEditor.Build.cs
+++ b/Engine/Plugins/Media/MediaPlate/Source/MediaPlateEditor/MediaPlateEditor.Build.cs
@@ -33,6 +33,7 @@ public class MediaPlateEditor : ModuleRules
 				"ModelingComponents",
 				"MovieScene",
 				"PropertyEditor",
+				"SceneOutliner",
 				"Sequencer",
 				"Slate",
 				"SlateCore",
diff --git a/Engine/Plugins/Media/MediaPlate/Source/MediaPlateEditor/Private/Sequencer/MediaPlateTrackEditor.cpp b/Engine/Plugins/Media/MediaPlate/Source/MediaPlateEditor/Private/Sequencer/MediaPlateTrackEditor.cpp
--- a/Engine/Plugins/Media/MediaPlate/Source/MediaPlateEditor/Private/Sequencer/MediaPlateTrackEditor.cpp
+++ b/Engine/Plugins/Media/MediaPlate/Source/MediaPlateEditor/Private/Sequencer/MediaPlateTrackEditor.cpp
@@ -2,6 +2,7 @@
 
 #include "Sequencer/MediaPlateTrackEditor.h"
 
+#include "ActorTreeItem.h"
 #include "Framework/MultiBox/MultiBoxBuilder.h"
 #include "ISequencer.h"
 #include "MediaTexture.h"
@@ -9,8 +10,11 @@
 #include "MediaPlateComponent.h"
 #include "MediaPlayer.h"
 #include "MediaPlaylist.h"
+#include "Modules/ModuleManager.h"
 #include "MovieSceneMediaSection.h"
 #include "MovieSceneMediaTrack.h"
+#include "SceneOutlinerModule.h"
+#include "Sequencer/MediaTrackEditor.h"
 
 #define LOCTEXT_NAMESPACE "FMediaPlateTrackEditor"
 
@@ -30,10 +34,13 @@ FMediaPlateTrackEditor::FMediaPlateTrackEditor(TSharedRef<ISequencer> InSequence
 	, bGetDurationDelay(false)
 {
 	OnActorAddedToSequencerHandle = InSequencer->OnActorAddedToSequencer().AddRaw(this, &FMediaPlateTrackEditor::HandleActorAdded);
+	FMediaTrackEditor::OnBuildOutlinerEditWidget.AddRaw(this,
+		&FMediaPlateTrackEditor::OnBuildOutlinerEditWidget);
 }
 
 FMediaPlateTrackEditor::~FMediaPlateTrackEditor()
 {
+	FMediaTrackEditor::OnBuildOutlinerEditWidget.RemoveAll(this);
 }
 
 void FMediaPlateTrackEditor::BuildObjectBindingContextMenu(FMenuBuilder& MenuBuilder, const TArray<FGuid>& ObjectBindings, const UClass* ObjectClass)
@@ -280,4 +287,58 @@ bool FMediaPlateTrackEditor::GetDuration(
 	return bIsDone;
 }
 
+void FMediaPlateTrackEditor::OnBuildOutlinerEditWidget(FMenuBuilder& MenuBuilder)
+{
+	// Add media plate sub menu.
+	MenuBuilder.AddSubMenu(
+		LOCTEXT("MediaPlate", "Media Plate"),
+		LOCTEXT("MediaPlateTooltip", "Add media from a media plate."),
+		FNewMenuDelegate::CreateLambda([this](FMenuBuilder& MenuBuilder)
+		{
+			// Actor picker options.
+			FSceneOutlinerInitializationOptions InitOptions;
+			InitOptions.bShowHeaderRow = false;
+			InitOptions.bShowSearchBox = true;
+			InitOptions.bShowCreateNewFolder = false;
+			InitOptions.bFocusSearchBoxWhenOpened = true;
+			InitOptions.ColumnMap.Add(FSceneOutlinerBuiltInColumnTypes::Label(), FSceneOutlinerColumnInfo(ESceneOutlinerColumnVisibility::Visible, 0));
+			InitOptions.Filters->AddFilterPredicate<FActorTreeItem>(
+				FActorTreeItem::FFilterPredicate::CreateLambda([](const AActor* Actor)
+				{
+					return Actor && Actor->IsA<AMediaPlate>();
+				}
+				));
+
+			// Create actor picker.
+			FSceneOutlinerModule& SceneOutlinerModule = 
+				FModuleManager::LoadModuleChecked<FSceneOutlinerModule>("SceneOutliner");
+			TSharedRef<SBox> Picker = SNew(SBox)
+				.WidthOverride(300.0f)
+				.HeightOverride(300.f)
+				[
+					SceneOutlinerModule.CreateActorPicker(InitOptions,
+						FOnActorPicked::CreateLambda([this](AActor* Actor)
+						{
+							AddMediaPlateToSequencer(Actor);
+						}))
+				];
+
+			MenuBuilder.AddWidget(Picker, FText::GetEmpty(), true);
+		}));
+}
+
+void FMediaPlateTrackEditor::AddMediaPlateToSequencer(AActor* Actor)
+{
+	if (Actor != nullptr)
+	{
+		TSharedPtr<ISequencer> SequencerPtr = GetSequencer();
+		if (SequencerPtr.IsValid())
+		{
+			TArray<TWeakObjectPtr<AActor>> ActorArray;
+			ActorArray.Add(Actor);
+			SequencerPtr->AddActors(ActorArray);
+		}
+	}
+}
+
 #undef LOCTEXT_NAMESPACE
diff --git a/Engine/Plugins/Media/MediaPlate/Source/MediaPlateEditor/Private/Sequencer/MediaPlateTrackEditor.h b/Engine/Plugins/Media/MediaPlate/Source/MediaPlateEditor/Private/Sequencer/MediaPlateTrackEditor.h
--- a/Engine/Plugins/Media/MediaPlate/Source/MediaPlateEditor/Private/Sequencer/MediaPlateTrackEditor.h
+++ b/Engine/Plugins/Media/MediaPlate/Source/MediaPlateEditor/Private/Sequencer/MediaPlateTrackEditor.h
@@ -99,6 +99,20 @@ private:
 	bool GetDuration(TStrongObjectPtr<UMediaPlayer>& MediaPlayer,
 		TWeakObjectPtr<UMovieSceneSection>& NewSection);
 
+	/**
+	 * Adds media plates to the MediaTrackEditor widget.
+	 * 
+	 * @param MenuBuilder		Menu to add to.
+	 */
+	void OnBuildOutlinerEditWidget(FMenuBuilder& MenuBuilder);
+
+	/**
+	 * Adds a media plate to a Sequencer track.
+	 * 
+	 * @param Actor			Media plate actor to add.
+	 */
+	void AddMediaPlateToSequencer(AActor* Actor);
+
 	/** Handle to our delegate. */
 	FDelegateHandle OnActorAddedToSequencerHandle;
 
