MediaPlate: Moved component blueprint functions to a different category to differentiate from the actor ones.
#rnx

diff --git a/Engine/Plugins/Media/MediaPlate/Source/MediaPlate/Public/MediaPlateComponent.h b/Engine/Plugins/Media/MediaPlate/Source/MediaPlate/Public/MediaPlateComponent.h
--- a/Engine/Plugins/Media/MediaPlate/Source/MediaPlate/Public/MediaPlateComponent.h
+++ b/Engine/Plugins/Media/MediaPlate/Source/MediaPlate/Public/MediaPlateComponent.h
@@ -41,25 +41,25 @@ public:
 	/**
 	 * Call this get our media player.
 	 */
-	UFUNCTION(BlueprintCallable, Category = "Media|MediaPlate")
+	UFUNCTION(BlueprintCallable, Category = "Media|MediaPlateComponent")
 	UMediaPlayer* GetMediaPlayer();
 
 	/**
 	 * Call this get our media texture.
 	 */
-	UFUNCTION(BlueprintCallable, Category = "Media|MediaPlate")
+	UFUNCTION(BlueprintCallable, Category = "Media|MediaPlateComponent")
 	UMediaTexture* GetMediaTexture();
 
 	/**
 	 * Call this to start playing.
 	 */
-	UFUNCTION(BlueprintCallable, Category = "Media|MediaPlate")
+	UFUNCTION(BlueprintCallable, Category = "Media|MediaPlateComponent")
 	void Play();
 
 	/**
 	 * Call this to stop playing.
 	 */
-	UFUNCTION(BlueprintCallable, Category = "Media|MediaPlate")
+	UFUNCTION(BlueprintCallable, Category = "Media|MediaPlateComponent")
 	void Stop();
 
 	/** If set then start playing right away. */
@@ -109,15 +109,15 @@ public:
 	EMediaTextureVisibleMipsTiles VisibleMipsTilesCalculations;
 
 	/** Set the arc size in degrees used for visible mips and tiles calculations, specific to the sphere. */
-	UFUNCTION(BlueprintCallable, Category = "Media|MediaPlate")
+	UFUNCTION(BlueprintCallable, Category = "Media|MediaPlateComponent")
 	void SetMeshRange(FVector2D InMeshRange);
 
 	/** Return the arc size in degrees used for visible mips and tiles calculations, specific to the sphere. */
-	UFUNCTION(BlueprintCallable, Category = "Media|MediaPlate")
+	UFUNCTION(BlueprintCallable, Category = "Media|MediaPlateComponent")
 	FVector2D GetMeshRange() const { return MeshRange; }
 
 	/** Call this to set bPlayOnlyWhenVisible. */
-	UFUNCTION(BlueprintCallable, Category = "Media|MediaPlate")
+	UFUNCTION(BlueprintCallable, Category = "Media|MediaPlateComponent")
 	void SetPlayOnlyWhenVisible(bool bInPlayOnlyWhenVisible);
 
 	/**
@@ -142,13 +142,13 @@ public:
 	/**
 	 * Call this to get the aspect ratio of the screen.
 	 */
-	UFUNCTION(BlueprintCallable, Category = "Media|MediaPlate")
+	UFUNCTION(BlueprintCallable, Category = "Media|MediaPlateComponent")
 	float GetLetterboxAspectRatio() { return LetterboxAspectRatio; }
 
 	/**
 	 * Call this to set the aspect ratio of the screen.
 	 */
-	UFUNCTION(BlueprintCallable, Category = "Media|MediaPlate")
+	UFUNCTION(BlueprintCallable, Category = "Media|MediaPlateComponent")
 	void SetLetterboxAspectRatio(float AspectRatio);
 
 	/**
