MediaPlate: Added getter and setter for loop.
#rnx

diff --git a/Engine/Plugins/Media/MediaPlate/Source/MediaPlate/Private/MediaPlateComponent.cpp b/Engine/Plugins/Media/MediaPlate/Source/MediaPlate/Private/MediaPlateComponent.cpp
--- a/Engine/Plugins/Media/MediaPlate/Source/MediaPlate/Private/MediaPlateComponent.cpp
+++ b/Engine/Plugins/Media/MediaPlate/Source/MediaPlate/Private/MediaPlateComponent.cpp
@@ -237,6 +237,23 @@ void UMediaPlateComponent::Stop()
 	bWantsToPlayWhenVisible = false;
 }
 
+bool UMediaPlateComponent::GetLoop()
+{
+	return bLoop;
+}
+
+void UMediaPlateComponent::SetLoop(bool bInLoop)
+{
+	if (bLoop != bInLoop)
+	{
+		bLoop = bInLoop;
+		if (MediaPlayer != nullptr)
+		{
+			MediaPlayer->SetLooping(bLoop);
+		}
+	}
+}
+
 void UMediaPlateComponent::SetMeshRange(FVector2D InMeshRange)
 {
 	MeshRange = InMeshRange;
diff --git a/Engine/Plugins/Media/MediaPlate/Source/MediaPlate/Public/MediaPlateComponent.h b/Engine/Plugins/Media/MediaPlate/Source/MediaPlate/Public/MediaPlateComponent.h
--- a/Engine/Plugins/Media/MediaPlate/Source/MediaPlate/Public/MediaPlateComponent.h
+++ b/Engine/Plugins/Media/MediaPlate/Source/MediaPlate/Public/MediaPlateComponent.h
@@ -62,14 +62,22 @@ public:
 	UFUNCTION(BlueprintCallable, Category = "Media|MediaPlateComponent")
 	void Stop();
 
+	/**
+	 * Call this to see if we want to loop.
+	 */
+	UFUNCTION(BlueprintGetter)
+	bool GetLoop();
+
+	/**
+	 * Call this enable/disable looping.
+	 */
+	UFUNCTION(BlueprintSetter)
+	void SetLoop(bool bInLoop);
+
 	/** If set then start playing right away. */
 	UPROPERTY(EditAnywhere, Category = "MediaPlate")
 	bool bAutoPlay;
 
-	/** If set then loop when we reach the end. */
-	UPROPERTY(EditAnywhere, Category = "MediaPlate")
-	bool bLoop;
-
 	/** If set then enable audio. */
 	UPROPERTY(EditAnywhere, Category = "MediaPlate")
 	bool bEnableAudio = false;
@@ -166,6 +174,10 @@ private:
 	UPROPERTY(EditAnywhere, BlueprintReadOnly, Category = "MediaPlate", meta = (AllowPrivateAccess = true))
 	bool bPlayOnlyWhenVisible = false;
 
+	/** If set then loop when we reach the end. */
+	UPROPERTY(EditAnywhere, Blueprintgetter = GetLoop, BlueprintSetter = SetLoop, Category = "MediaPlate", meta = (AllowPrivateAccess = true))
+	bool bLoop;
+
 	/** If > 0, then this is the aspect ratio of our screen and 
 	 * letterboxes will be added if the media is smaller than the screen. */
 	UPROPERTY()
