MediaPlate: Play split into Open and Play to match MediaPlayer.
#rnx

diff --git a/Engine/Plugins/Media/MediaPlate/Source/MediaPlate/Private/MediaPlateComponent.cpp b/Engine/Plugins/Media/MediaPlate/Source/MediaPlate/Private/MediaPlateComponent.cpp
--- a/Engine/Plugins/Media/MediaPlate/Source/MediaPlate/Private/MediaPlateComponent.cpp
+++ b/Engine/Plugins/Media/MediaPlate/Source/MediaPlate/Private/MediaPlateComponent.cpp
@@ -120,7 +120,7 @@ void UMediaPlateComponent::BeginPlay()
 	// Start playing?
 	if (bAutoPlay)
 	{
-		Play();
+		Open();
 	}
 }
 
@@ -180,7 +180,7 @@ UMediaTexture* UMediaPlateComponent::GetMediaTexture()
 	return MediaTexture;
 }
 
-void UMediaPlateComponent::Play()
+void UMediaPlateComponent::Open()
 {
 	if ((bPlayOnlyWhenVisible == false) || (IsVisible()))
 	{
@@ -226,6 +226,15 @@ void UMediaPlateComponent::Play()
 	}
 }
 
+
+void UMediaPlateComponent::Play()
+{
+	if (MediaPlayer != nullptr)
+	{
+		MediaPlayer->Play();
+	}
+}
+
 void UMediaPlateComponent::Stop()
 {
 	if (MediaPlayer != nullptr)
@@ -318,7 +327,8 @@ bool UMediaPlateComponent::PlayMediaSource(UMediaSource* InMediaSource)
 			// Play the source.
 			FMediaPlayerOptions Options;
 			Options.SeekTime = FTimespan::FromSeconds(StartTime);
-			Options.PlayOnOpen = EMediaPlayerOptionBooleanOverride::Enabled;
+			Options.PlayOnOpen = bPlayOnOpen ? EMediaPlayerOptionBooleanOverride::Enabled :
+				EMediaPlayerOptionBooleanOverride::Disabled;
 			Options.Loop = bLoop ? EMediaPlayerOptionBooleanOverride::Enabled :
 				EMediaPlayerOptionBooleanOverride::Disabled;
 			bIsPlaying = MediaPlayer->OpenSourceWithOptions(InMediaSource, Options);
@@ -441,7 +451,7 @@ void UMediaPlateComponent::ResumeWhenVisible()
 		}
 		else if (bWantsToPlayWhenVisible)
 		{
-			Play();
+			Open();
 			FTimespan PlayTime = GetResumeTime();
 			MediaPlayer->Seek(PlayTime);
 		}
@@ -573,7 +583,7 @@ void UMediaPlateComponent::OnVisibleMipsTilesCalculationsChange()
 			if (MediaPlayer->IsPlaying())
 			{
 				MediaPlayer->Close();
-				Play();
+				Open();
 			}
 		}
 	}
diff --git a/Engine/Plugins/Media/MediaPlate/Source/MediaPlate/Public/MediaPlateComponent.h b/Engine/Plugins/Media/MediaPlate/Source/MediaPlate/Public/MediaPlateComponent.h
--- a/Engine/Plugins/Media/MediaPlate/Source/MediaPlate/Public/MediaPlateComponent.h
+++ b/Engine/Plugins/Media/MediaPlate/Source/MediaPlate/Public/MediaPlateComponent.h
@@ -50,8 +50,15 @@ public:
 	UFUNCTION(BlueprintCallable, Category = "Media|MediaPlateComponent")
 	UMediaTexture* GetMediaTexture();
 
+	/**
+	 * Call this to open the media.
+	 */
+	UFUNCTION(BlueprintCallable, Category = "Media|MediaPlateComponent")
+	void Open();
+
 	/**
 	 * Call this to start playing.
+	 * Open must be called before this.
 	 */
 	UFUNCTION(BlueprintCallable, Category = "Media|MediaPlateComponent")
 	void Play();
@@ -74,7 +81,11 @@ public:
 	UFUNCTION(BlueprintSetter)
 	void SetLoop(bool bInLoop);
 
-	/** If set then start playing right away. */
+	/** If set then play when opening the media. */
+	UPROPERTY(EditAnywhere, BlueprintReadWrite, Category = "MediaPlate")
+	bool bPlayOnOpen = true;
+
+	/** If set then start playing when this object is active. */
 	UPROPERTY(EditAnywhere, Category = "MediaPlate")
 	bool bAutoPlay;
 
diff --git a/Engine/Plugins/Media/MediaPlate/Source/MediaPlateEditor/Private/MediaPlateCustomization.cpp b/Engine/Plugins/Media/MediaPlate/Source/MediaPlateEditor/Private/MediaPlateCustomization.cpp
--- a/Engine/Plugins/Media/MediaPlate/Source/MediaPlateEditor/Private/MediaPlateCustomization.cpp
+++ b/Engine/Plugins/Media/MediaPlate/Source/MediaPlateEditor/Private/MediaPlateCustomization.cpp
@@ -259,7 +259,7 @@ void FMediaPlateCustomization::CustomizeDetails(IDetailLayoutBuilder& DetailBuil
 										((MediaPlayer->IsPaused()) || 
 											(MediaPlayer->IsPlaying() && (MediaPlayer->GetRate() != 1.0f))))
 									{
-										MediaPlayer->Play();
+										MediaPlate->Play();
 									}
 									else
 									{
@@ -271,7 +271,7 @@ void FMediaPlateCustomization::CustomizeDetails(IDetailLayoutBuilder& DetailBuil
 										}
 
 										// Play the media.
-										MediaPlate->Play();
+										MediaPlate->Open();
 									}
 								}
 							}
diff --git a/Engine/Plugins/Media/MediaPlate/Source/MediaPlateEditor/Private/Toolkits/MediaPlateEditorToolkit.cpp b/Engine/Plugins/Media/MediaPlate/Source/MediaPlateEditor/Private/Toolkits/MediaPlateEditorToolkit.cpp
--- a/Engine/Plugins/Media/MediaPlate/Source/MediaPlateEditor/Private/Toolkits/MediaPlateEditorToolkit.cpp
+++ b/Engine/Plugins/Media/MediaPlate/Source/MediaPlateEditor/Private/Toolkits/MediaPlateEditorToolkit.cpp
@@ -268,7 +268,7 @@ void FMediaPlateEditorToolkit::BindCommands()
 
 	ToolkitCommands->MapAction(
 		Commands.OpenMedia,
-		FExecuteAction::CreateLambda([this] { MediaPlate->Play(); }),
+		FExecuteAction::CreateLambda([this] { MediaPlate->Open(); }),
 		FCanExecuteAction::CreateLambda([this] { return true; })
 	);
 
