MediaPlate: Added pause command.
#rnx

diff --git a/Engine/Plugins/Media/MediaPlate/Source/MediaPlate/Private/MediaPlateComponent.cpp b/Engine/Plugins/Media/MediaPlate/Source/MediaPlate/Private/MediaPlateComponent.cpp
--- a/Engine/Plugins/Media/MediaPlate/Source/MediaPlate/Private/MediaPlateComponent.cpp
+++ b/Engine/Plugins/Media/MediaPlate/Source/MediaPlate/Private/MediaPlateComponent.cpp
@@ -226,7 +226,6 @@ void UMediaPlateComponent::Open()
 	}
 }
 
-
 void UMediaPlateComponent::Play()
 {
 	if (MediaPlayer != nullptr)
@@ -235,6 +234,14 @@ void UMediaPlateComponent::Play()
 	}
 }
 
+void UMediaPlateComponent::Pause()
+{
+	if (MediaPlayer != nullptr)
+	{
+		MediaPlayer->Pause();
+	}
+}
+
 void UMediaPlateComponent::Stop()
 {
 	if (MediaPlayer != nullptr)
diff --git a/Engine/Plugins/Media/MediaPlate/Source/MediaPlate/Public/MediaPlateComponent.h b/Engine/Plugins/Media/MediaPlate/Source/MediaPlate/Public/MediaPlateComponent.h
--- a/Engine/Plugins/Media/MediaPlate/Source/MediaPlate/Public/MediaPlateComponent.h
+++ b/Engine/Plugins/Media/MediaPlate/Source/MediaPlate/Public/MediaPlateComponent.h
@@ -63,6 +63,13 @@ public:
 	UFUNCTION(BlueprintCallable, Category = "Media|MediaPlateComponent")
 	void Play();
 
+	/**
+	 * Call this to pause playback.
+	 * Play can be called to resume playback.
+	 */
+	UFUNCTION(BlueprintCallable, Category = "Media|MediaPlateComponent")
+	void Pause();
+
 	/**
 	 * Call this to stop playing.
 	 */
diff --git a/Engine/Plugins/Media/MediaPlate/Source/MediaPlateEditor/Private/MediaPlateCustomization.cpp b/Engine/Plugins/Media/MediaPlate/Source/MediaPlateEditor/Private/MediaPlateCustomization.cpp
--- a/Engine/Plugins/Media/MediaPlate/Source/MediaPlateEditor/Private/MediaPlateCustomization.cpp
+++ b/Engine/Plugins/Media/MediaPlate/Source/MediaPlateEditor/Private/MediaPlateCustomization.cpp
@@ -297,12 +297,8 @@ void FMediaPlateCustomization::CustomizeDetails(IDetailLayoutBuilder& DetailBuil
 						{
 							UMediaPlateComponent* MediaPlate = MediaPlatePtr.Get();
 							if (MediaPlate != nullptr)
-							{
-								UMediaPlayer* MediaPlayer = MediaPlate->GetMediaPlayer();
-								if (MediaPlayer != nullptr)
-								{
-									MediaPlayer->Pause();
-								}
+							{							
+								MediaPlate->Pause();
 							}
 						}
 						return FReply::Handled();
diff --git a/Engine/Plugins/Media/MediaPlate/Source/MediaPlateEditor/Private/Toolkits/MediaPlateEditorToolkit.cpp b/Engine/Plugins/Media/MediaPlate/Source/MediaPlateEditor/Private/Toolkits/MediaPlateEditorToolkit.cpp
--- a/Engine/Plugins/Media/MediaPlate/Source/MediaPlateEditor/Private/Toolkits/MediaPlateEditorToolkit.cpp
+++ b/Engine/Plugins/Media/MediaPlate/Source/MediaPlateEditor/Private/Toolkits/MediaPlateEditorToolkit.cpp
@@ -274,7 +274,7 @@ void FMediaPlateEditorToolkit::BindCommands()
 
 	ToolkitCommands->MapAction(
 		Commands.PauseMedia,
-		FExecuteAction::CreateLambda([this]{ MediaPlate->GetMediaPlayer()->Pause(); }),
+		FExecuteAction::CreateLambda([this]{ MediaPlate->Pause(); }),
 		FCanExecuteAction::CreateLambda([this]{
 			TObjectPtr<UMediaPlayer> MediaPlayer = MediaPlate->GetMediaPlayer();
 			return MediaPlayer->CanPause() && !MediaPlayer->IsPaused();
