MediaPlate: Renamed Stop to Close to match media player.
#rnx

diff --git a/Engine/Plugins/Media/MediaPlate/Source/MediaPlate/Private/MediaPlateComponent.cpp b/Engine/Plugins/Media/MediaPlate/Source/MediaPlate/Private/MediaPlateComponent.cpp
--- a/Engine/Plugins/Media/MediaPlate/Source/MediaPlate/Private/MediaPlateComponent.cpp
+++ b/Engine/Plugins/Media/MediaPlate/Source/MediaPlate/Private/MediaPlateComponent.cpp
@@ -242,7 +242,7 @@ void UMediaPlateComponent::Pause()
 	}
 }
 
-void UMediaPlateComponent::Stop()
+void UMediaPlateComponent::Close()
 {
 	if (MediaPlayer != nullptr)
 	{
diff --git a/Engine/Plugins/Media/MediaPlate/Source/MediaPlate/Public/MediaPlateComponent.h b/Engine/Plugins/Media/MediaPlate/Source/MediaPlate/Public/MediaPlateComponent.h
--- a/Engine/Plugins/Media/MediaPlate/Source/MediaPlate/Public/MediaPlateComponent.h
+++ b/Engine/Plugins/Media/MediaPlate/Source/MediaPlate/Public/MediaPlateComponent.h
@@ -71,10 +71,10 @@ public:
 	void Pause();
 
 	/**
-	 * Call this to stop playing.
+	 * Call this to close the media.
 	 */
 	UFUNCTION(BlueprintCallable, Category = "Media|MediaPlateComponent")
-	void Stop();
+	void Close();
 
 	/**
 	 * Call this to see if we want to loop.
diff --git a/Engine/Plugins/Media/MediaPlate/Source/MediaPlateEditor/Private/MediaPlateCustomization.cpp b/Engine/Plugins/Media/MediaPlate/Source/MediaPlateEditor/Private/MediaPlateCustomization.cpp
--- a/Engine/Plugins/Media/MediaPlate/Source/MediaPlateEditor/Private/MediaPlateCustomization.cpp
+++ b/Engine/Plugins/Media/MediaPlate/Source/MediaPlateEditor/Private/MediaPlateCustomization.cpp
@@ -369,7 +369,7 @@ void FMediaPlateCustomization::CustomizeDetails(IDetailLayoutBuilder& DetailBuil
 								UMediaPlateComponent* MediaPlate = MediaPlatePtr.Get();
 								if (MediaPlate != nullptr)
 								{
-									MediaPlate->Stop();
+									MediaPlate->Close();
 								}
 							}
 							return FReply::Handled();
@@ -1026,7 +1026,7 @@ void FMediaPlateCustomization::StopMediaPlates()
 		UMediaPlateComponent* MediaPlate = MediaPlatePtr.Get();
 		if (MediaPlate != nullptr)
 		{
-			MediaPlate->Stop();
+			MediaPlate->Close();
 		}
 	}
 }
