MediaPlate: Added open button to customization.
#rnx

diff --git a/Engine/Plugins/Media/MediaPlate/Source/MediaPlateEditor/Private/MediaPlateCustomization.cpp b/Engine/Plugins/Media/MediaPlate/Source/MediaPlateEditor/Private/MediaPlateCustomization.cpp
--- a/Engine/Plugins/Media/MediaPlate/Source/MediaPlateEditor/Private/MediaPlateCustomization.cpp
+++ b/Engine/Plugins/Media/MediaPlate/Source/MediaPlateEditor/Private/MediaPlateCustomization.cpp
@@ -246,7 +246,7 @@ void FMediaPlateCustomization::CustomizeDetails(IDetailLayoutBuilder& DetailBuil
 				[
 					SNew(SButton)
 						.VAlign(VAlign_Center)
-						.OnClicked_Lambda([this]() -> FReply
+						.IsEnabled_Lambda([this]
 						{
 							for (TWeakObjectPtr<UMediaPlateComponent>& MediaPlatePtr : MediaPlatesList)
 							{
@@ -255,24 +255,21 @@ void FMediaPlateCustomization::CustomizeDetails(IDetailLayoutBuilder& DetailBuil
 								{
 									// Is the player paused or fast forwarding/rewinding?
 									UMediaPlayer* MediaPlayer = MediaPlate->GetMediaPlayer();
-									if ((MediaPlayer != nullptr) &&
-										((MediaPlayer->IsPaused()) || 
-											(MediaPlayer->IsPlaying() && (MediaPlayer->GetRate() != 1.0f))))
-									{
-										MediaPlate->Play();
-									}
-									else
-									{
-										// Tell the editor module that this media plate is playing.
-										FMediaPlateEditorModule* EditorModule = FModuleManager::LoadModulePtr<FMediaPlateEditorModule>("MediaPlateEditor");
-										if (EditorModule != nullptr)
-										{
-											EditorModule->MediaPlateStartedPlayback(MediaPlate);
-										}
-
-										// Play the media.
-										MediaPlate->Open();
-									}
+									return ((MediaPlayer != nullptr) &&
+										(MediaPlayer->IsReady()) &&
+										((!MediaPlayer->IsPlaying() || (MediaPlayer->GetRate() != 1.0f))));
+								}
+							}
+							return false;
+						})
+						.OnClicked_Lambda([this]() -> FReply
+						{
+							for (TWeakObjectPtr<UMediaPlateComponent>& MediaPlatePtr : MediaPlatesList)
+							{
+								UMediaPlateComponent* MediaPlate = MediaPlatePtr.Get();
+								if (MediaPlate != nullptr)
+								{
+									MediaPlate->Play();
 								}
 							}
 							return FReply::Handled();
@@ -355,7 +352,41 @@ void FMediaPlateCustomization::CustomizeDetails(IDetailLayoutBuilder& DetailBuil
 					]
 				]
 
-			// Stop button.
+			// Open button.
+			+ SHorizontalBox::Slot()
+				.VAlign(VAlign_Center)
+				.AutoWidth()
+				[
+					SNew(SButton)
+						.VAlign(VAlign_Center)
+						.OnClicked_Lambda([this]() -> FReply
+						{
+							for (TWeakObjectPtr<UMediaPlateComponent>& MediaPlatePtr : MediaPlatesList)
+							{
+								UMediaPlateComponent* MediaPlate = MediaPlatePtr.Get();
+								if (MediaPlate != nullptr)
+								{
+									// Tell the editor module that this media plate is playing.
+									FMediaPlateEditorModule* EditorModule = FModuleManager::LoadModulePtr<FMediaPlateEditorModule>("MediaPlateEditor");
+									if (EditorModule != nullptr)
+									{
+										EditorModule->MediaPlateStartedPlayback(MediaPlate);
+									}
+
+									// Play the media.
+									MediaPlate->Open();
+								}
+							}
+							return FReply::Handled();
+						})
+						[
+							SNew(SImage)
+								.ColorAndOpacity(FSlateColor::UseForeground())
+								.Image(Style->GetBrush("MediaPlateEditor.OpenMedia.Small"))
+						]
+				]
+
+			// Close button.
 			+ SHorizontalBox::Slot()
 				.VAlign(VAlign_Center)
 				.AutoWidth()
diff --git a/Engine/Plugins/Media/MediaPlate/Source/MediaPlateEditor/Private/MediaPlateEditorStyle.cpp b/Engine/Plugins/Media/MediaPlate/Source/MediaPlateEditor/Private/MediaPlateEditorStyle.cpp
--- a/Engine/Plugins/Media/MediaPlate/Source/MediaPlateEditor/Private/MediaPlateEditorStyle.cpp
+++ b/Engine/Plugins/Media/MediaPlate/Source/MediaPlateEditor/Private/MediaPlateEditorStyle.cpp
@@ -55,6 +55,7 @@ FMediaPlateEditorStyle::FMediaPlateEditorStyle()
 	Set("MediaPlateEditor.NextMedia", new IMAGE_BRUSH("icon_step_40x", Icon40x40));
 	Set("MediaPlateEditor.NextMedia.Small", new IMAGE_BRUSH("icon_step_40x", Icon20x20));
 	Set("MediaPlateEditor.OpenMedia", new IMAGE_BRUSH("tab_media_16x", Icon40x40));
+	Set("MediaPlateEditor.OpenMedia.Small", new IMAGE_BRUSH("tab_media_16x", Icon20x20));
 	Set("MediaPlateEditor.PauseMedia", new IMAGE_BRUSH("icon_pause_40x", Icon40x40));
 	Set("MediaPlateEditor.PauseMedia.Small", new IMAGE_BRUSH("icon_pause_40x", Icon20x20));
 	Set("MediaPlateEditor.PlayMedia", new IMAGE_BRUSH("icon_play_40x", Icon40x40));
