MediaPlate: Added seek and rewind commands.
#rnx

diff --git a/Engine/Plugins/Media/MediaPlate/Source/MediaPlate/Private/MediaPlateComponent.cpp b/Engine/Plugins/Media/MediaPlate/Source/MediaPlate/Private/MediaPlateComponent.cpp
--- a/Engine/Plugins/Media/MediaPlate/Source/MediaPlate/Private/MediaPlateComponent.cpp
+++ b/Engine/Plugins/Media/MediaPlate/Source/MediaPlate/Private/MediaPlateComponent.cpp
@@ -242,6 +242,21 @@ void UMediaPlateComponent::Pause()
 	}
 }
 
+bool UMediaPlateComponent::Rewind()
+{
+	return Seek(FTimespan::Zero());
+}
+
+bool UMediaPlateComponent::Seek(const FTimespan& Time)
+{
+	if (MediaPlayer != nullptr)
+	{
+		return MediaPlayer->Seek(Time);
+	}
+
+	return false;
+}
+
 void UMediaPlateComponent::Close()
 {
 	if (MediaPlayer != nullptr)
diff --git a/Engine/Plugins/Media/MediaPlate/Source/MediaPlate/Public/MediaPlateComponent.h b/Engine/Plugins/Media/MediaPlate/Source/MediaPlate/Public/MediaPlateComponent.h
--- a/Engine/Plugins/Media/MediaPlate/Source/MediaPlate/Public/MediaPlateComponent.h
+++ b/Engine/Plugins/Media/MediaPlate/Source/MediaPlate/Public/MediaPlateComponent.h
@@ -70,6 +70,25 @@ public:
 	UFUNCTION(BlueprintCallable, Category = "Media|MediaPlateComponent")
 	void Pause();
 
+	/**
+	 * Rewinds the media to the beginning.
+	 *
+	 * This is the same as seeking to zero time.
+	 *
+	 * @return				True if rewinding, false otherwise.
+	 */
+	UFUNCTION(BlueprintCallable, Category = "Media|MediaPlateComponent")
+	bool Rewind();
+
+	/**
+	 * Call this to seek to the specified playback time.
+	 * 
+	 * @param Time			Time to seek to.
+	 * @return				True on success, false otherwise.
+	 */
+	UFUNCTION(BlueprintCallable, Category = "Media|MediaPlateComponent")
+	bool Seek(const FTimespan& Time);
+
 	/**
 	 * Call this to close the media.
 	 */
