MediaPlate: Plays through the playlist now.

#rnx

diff --git a/Engine/Plugins/Media/MediaPlate/Source/MediaPlate/Private/MediaPlateComponent.cpp b/Engine/Plugins/Media/MediaPlate/Source/MediaPlate/Private/MediaPlateComponent.cpp
--- a/Engine/Plugins/Media/MediaPlate/Source/MediaPlate/Private/MediaPlateComponent.cpp
+++ b/Engine/Plugins/Media/MediaPlate/Source/MediaPlate/Private/MediaPlateComponent.cpp
@@ -92,6 +92,7 @@ void UMediaPlateComponent::OnRegister()
 		MediaPlayer->SetLooping(false);
 		MediaPlayer->PlayOnOpen = false;
 	}
+	MediaPlayer->OnEndReached.AddUniqueDynamic(this, &UMediaPlateComponent::OnMediaEnd);
 
 	// Set up media texture.
 	if (MediaTexture != nullptr)
@@ -200,24 +201,6 @@ void UMediaPlateComponent::Open()
 		{
 			UE_LOG(LogMediaPlate, Warning, TEXT("Could not play anything."));
 		}
-		else
-		{
-			// Are we using automatic aspect ratio?
-			if ((bIsAspectRatioAuto) &&
-				(VisibleMipsTilesCalculations == EMediaTextureVisibleMipsTiles::Plane))
-			{
-				// Start the clock sink so we can tick.
-				IMediaModule* MediaModule = FModuleManager::LoadModulePtr<IMediaModule>("Media");
-				if (MediaModule != nullptr)
-				{
-					if (ClockSink.IsValid() == false)
-					{
-						ClockSink = MakeShared<FMediaComponentClockSink, ESPMode::ThreadSafe>(this);
-					}
-					MediaModule->GetClock().AddSink(ClockSink.ToSharedRef());
-				}
-			}
-		}
 	}
 	else
 	{
@@ -266,6 +249,7 @@ void UMediaPlateComponent::Close()
 
 	StopClockSink();
 	bWantsToPlayWhenVisible = false;
+	PlaylistIndex = 0;
 }
 
 bool UMediaPlateComponent::GetLoop()
@@ -354,6 +338,26 @@ bool UMediaPlateComponent::PlayMediaSource(UMediaSource* InMediaSource)
 			Options.Loop = bLoop ? EMediaPlayerOptionBooleanOverride::Enabled :
 				EMediaPlayerOptionBooleanOverride::Disabled;
 			bIsPlaying = MediaPlayer->OpenSourceWithOptions(InMediaSource, Options);
+
+			// Did we play anything?
+			if (bIsPlaying)
+			{
+				// Are we using automatic aspect ratio?
+				if ((bIsAspectRatioAuto) &&
+					(VisibleMipsTilesCalculations == EMediaTextureVisibleMipsTiles::Plane))
+				{
+					// Start the clock sink so we can tick.
+					IMediaModule* MediaModule = FModuleManager::LoadModulePtr<IMediaModule>("Media");
+					if (MediaModule != nullptr)
+					{
+						if (ClockSink.IsValid() == false)
+						{
+							ClockSink = MakeShared<FMediaComponentClockSink, ESPMode::ThreadSafe>(this);
+						}
+						MediaModule->GetClock().AddSink(ClockSink.ToSharedRef());
+					}
+				}
+			}
 		}
 	}
 
@@ -591,6 +595,26 @@ void UMediaPlateComponent::RemoveLetterboxes()
 	Letterboxes.Empty();
 }
 
+
+void UMediaPlateComponent::OnMediaEnd()
+{
+	StopClockSink();
+
+	// Do we have a playlist>
+	if ((MediaPlaylist != nullptr) && (MediaPlaylist->Num() > 1))
+	{
+		// Get the next media to play.
+		UMediaSource* NextSource = MediaPlaylist->GetNext(PlaylistIndex);
+		if (NextSource != nullptr)
+		{
+			if (PlaylistIndex != 0)
+			{
+				PlayMediaSource(NextSource);
+			}
+		}
+	}
+}
+
 #if WITH_EDITOR
 
 void UMediaPlateComponent::OnVisibleMipsTilesCalculationsChange()
diff --git a/Engine/Plugins/Media/MediaPlate/Source/MediaPlate/Public/MediaPlateComponent.h b/Engine/Plugins/Media/MediaPlate/Source/MediaPlate/Public/MediaPlateComponent.h
--- a/Engine/Plugins/Media/MediaPlate/Source/MediaPlate/Public/MediaPlateComponent.h
+++ b/Engine/Plugins/Media/MediaPlate/Source/MediaPlate/Public/MediaPlateComponent.h
@@ -142,6 +142,10 @@ public:
 	UPROPERTY(EditAnywhere, Category = "MediaPlate")
 	TObjectPtr<UMediaPlaylist> MediaPlaylist;
 
+	/** The current index of the source in the play list being played. */
+	UPROPERTY(BlueprintReadWrite, Category = "MediaPlate")
+	int32 PlaylistIndex = 0;
+
 	/** Override the default cache settings. */
 	UPROPERTY(EditAnywhere, AdvancedDisplay, Category = "MediaPlate")
 	FMediaSourceCacheSettings CacheSettings;
@@ -293,4 +297,10 @@ private:
 	 */
 	void RemoveLetterboxes();
 
+	/**
+	 * Called by the media player when the video ends.
+	 */
+	UFUNCTION()
+	void OnMediaEnd();
+
 };
