MediaPlate: Explicilty shows the playlist as an external asset, and an empty playlist now denotes an internal one.
#rnx

diff --git a/Engine/Plugins/Media/MediaPlate/Source/MediaPlate/Public/MediaPlateComponent.h b/Engine/Plugins/Media/MediaPlate/Source/MediaPlate/Public/MediaPlateComponent.h
--- a/Engine/Plugins/Media/MediaPlate/Source/MediaPlate/Public/MediaPlateComponent.h
+++ b/Engine/Plugins/Media/MediaPlate/Source/MediaPlate/Public/MediaPlateComponent.h
@@ -139,7 +139,7 @@ public:
 	TArray<TObjectPtr<UStaticMeshComponent>> Letterboxes;
 
 	/** What media playlist to play. */
-	UPROPERTY(EditAnywhere, Category = "MediaPlate")
+	UPROPERTY(BlueprintReadWrite, Category = "MediaPlate")
 	TObjectPtr<UMediaPlaylist> MediaPlaylist;
 
 	/** The current index of the source in the play list being played. */
diff --git a/Engine/Plugins/Media/MediaPlate/Source/MediaPlateEditor/Private/MediaPlateCustomization.cpp b/Engine/Plugins/Media/MediaPlate/Source/MediaPlateEditor/Private/MediaPlateCustomization.cpp
--- a/Engine/Plugins/Media/MediaPlate/Source/MediaPlateEditor/Private/MediaPlateCustomization.cpp
+++ b/Engine/Plugins/Media/MediaPlate/Source/MediaPlateEditor/Private/MediaPlateCustomization.cpp
@@ -113,11 +113,24 @@ void FMediaPlateCustomization::CustomizeDetails(IDetailLayoutBuilder& DetailBuil
 	// Create playlist group.
 	IDetailGroup& PlaylistGroup = MediaPlateCategory.AddGroup(TEXT("Playlist"),
 		LOCTEXT("Playlist", "Playlist"));
-	TSharedRef<IPropertyHandle> PropertyHandle = DetailBuilder.GetProperty(
-		GET_MEMBER_NAME_CHECKED(UMediaPlateComponent, MediaPlaylist));
-	PlaylistGroup.HeaderProperty(PropertyHandle);
-	PropertyHandle->SetOnPropertyValueChanged(FSimpleDelegate::CreateSP(this,
-		&FMediaPlateCustomization::OnPlaylistChanged));
+
+	// Add playlist.
+	PlaylistGroup.AddWidgetRow()
+		.NameContent()
+		[
+			SNew(STextBlock)
+				.Text(LOCTEXT("PlaylistAsset", "Playlist Asset"))
+				.ToolTipText(LOCTEXT("PlaylistAsset_Tooltip",
+					"The playlist asset to use. If none then an internal asset will be used."))
+				.Font(IDetailLayoutBuilder::GetDetailFont())
+		]
+		.ValueContent()
+		[
+			SNew(SObjectPropertyEntryBox)
+				.AllowedClass(UMediaPlaylist::StaticClass())
+				.ObjectPath(this, &FMediaPlateCustomization::GetPlaylistPath)
+				.OnObjectChanged(this, &FMediaPlateCustomization::OnPlaylistChanged)
+		];
 
 	// Add media source.
 	PlaylistGroup.AddWidgetRow()
@@ -903,9 +916,45 @@ FString FMediaPlateCustomization::GetMediaSourcePath() const
 	return Path;
 }
 
-void FMediaPlateCustomization::OnPlaylistChanged()
+FString FMediaPlateCustomization::GetPlaylistPath() const
+{
+	FString Path;
+
+	// Get the first media plate.
+	if (MediaPlatesList.Num() > 0)
+	{
+		UMediaPlateComponent* MediaPlate = MediaPlatesList[0].Get();
+		if (MediaPlate != nullptr)
+		{
+			UMediaPlaylist* Playlist = MediaPlate->MediaPlaylist;
+			if (Playlist != nullptr)
+			{
+				if (Playlist->GetOuter() != MediaPlate)
+				{
+					Path = Playlist->GetPathName();
+				}
+			}
+		}
+	}
+
+	return Path;
+}
+
+void FMediaPlateCustomization::OnPlaylistChanged(const FAssetData& AssetData)
 {
 	StopMediaPlates();
+
+	// Update the playlist.
+	UMediaPlaylist* Playlist = Cast<UMediaPlaylist>(AssetData.GetAsset());
+	for (TWeakObjectPtr<UMediaPlateComponent>& MediaPlatePtr : MediaPlatesList)
+	{
+		UMediaPlateComponent* MediaPlate = MediaPlatePtr.Get();
+		if (MediaPlate != nullptr)
+		{
+			MediaPlate->MediaPlaylist = Playlist;
+		}
+	}
+
 	UpdateMediaPath();
 }
 
diff --git a/Engine/Plugins/Media/MediaPlate/Source/MediaPlateEditor/Private/MediaPlateCustomization.h b/Engine/Plugins/Media/MediaPlate/Source/MediaPlateEditor/Private/MediaPlateCustomization.h
--- a/Engine/Plugins/Media/MediaPlate/Source/MediaPlateEditor/Private/MediaPlateCustomization.h
+++ b/Engine/Plugins/Media/MediaPlate/Source/MediaPlateEditor/Private/MediaPlateCustomization.h
@@ -161,10 +161,15 @@ private:
 	 */
 	FString GetMediaSourcePath() const;
 
+	/**
+	 * Gets the object path for the media playlist object.
+	 */
+	FString GetPlaylistPath() const;
+
 	/**
 	 * Called when the playlist changes.
 	 */
-	void OnPlaylistChanged();
+	void OnPlaylistChanged(const FAssetData& AssetData);
 
 	/**
 	 * Called when the media source widget changes.
