MediaPlate: Added UI option to select between using an asset or a file for the media source.

#rnx

diff --git a/Engine/Plugins/Media/MediaPlate/Source/MediaPlateEditor/Private/MediaPlateCustomization.cpp b/Engine/Plugins/Media/MediaPlate/Source/MediaPlateEditor/Private/MediaPlateCustomization.cpp
--- a/Engine/Plugins/Media/MediaPlate/Source/MediaPlateEditor/Private/MediaPlateCustomization.cpp
+++ b/Engine/Plugins/Media/MediaPlate/Source/MediaPlateEditor/Private/MediaPlateCustomization.cpp
@@ -32,6 +32,7 @@
 
 
 FMediaPlateCustomization::FMediaPlateCustomization()
+	: bIsMediaSourceAsset(false)
 {
 	PropertyChangeDelegate = FCoreUObjectDelegates::OnObjectPropertyChanged.AddRaw(this, &FMediaPlateCustomization::OnObjectPropertyChanged);
 }
@@ -104,7 +105,8 @@ void FMediaPlateCustomization::CustomizeDetails(IDetailLayoutBuilder& DetailBuil
 		}
 	}
 
-	// Set media path.
+	// Set up media source.
+	UpdateIsMediaSourceAsset();
 	UpdateMediaPath();
 
 	// Add mesh customization.
@@ -132,14 +134,49 @@ void FMediaPlateCustomization::CustomizeDetails(IDetailLayoutBuilder& DetailBuil
 				.OnObjectChanged(this, &FMediaPlateCustomization::OnPlaylistChanged)
 		];
 
-	// Add media source.
+	// Add media source type.
+	TAttribute<EVisibility> MediaSourceAssetVisibility(this, &FMediaPlateCustomization::ShouldShowMediaSourceAsset);
+	TAttribute<EVisibility> MediaSourceFileVisibility(this, &FMediaPlateCustomization::ShouldShowMediaSourceFile);
 	PlaylistGroup.AddWidgetRow()
 		.NameContent()
 		[
 			SNew(STextBlock)
-			.Text(LOCTEXT("MediaSource", "Media Source"))
-			.ToolTipText(LOCTEXT("MediaSource_ToolTip", "The Media Source to play."))
-			.Font(IDetailLayoutBuilder::GetDetailFont())
+				.Text(LOCTEXT("FirstItemInPlaylist", "First Item In Playlist"))
+				.ToolTipText(LOCTEXT("FirstItemInPlaylist_ToolTip",
+					"The type of the first item in the playlist.\n"
+					"File lets you select a media file like a .mp4.\n"
+					"Asset lets you select a Media Source asset."))
+				.Font(IDetailLayoutBuilder::GetDetailFont())
+		]
+		.ValueContent()
+		[
+			SNew(SSegmentedControl<bool>)
+				.Value_Lambda([this]()
+				{
+					return bIsMediaSourceAsset;
+				})
+				.OnValueChanged(this, &FMediaPlateCustomization::SetIsMediaSourceAsset)
+
+			+ SSegmentedControl<bool>::Slot(false)
+				.Text(LOCTEXT("File", "File"))
+				.ToolTip(LOCTEXT("File_ToolTip",
+					"Select this if you want to use a file like a .mp4."))
+
+			+ SSegmentedControl<bool>::Slot(true)
+				.Text(LOCTEXT("Asset", "Asset"))
+				.ToolTip(LOCTEXT("Asset_ToolTip",
+					"Select this if you want to use a Media Source asset."))
+		];
+
+	// Add media asset.
+	PlaylistGroup.AddWidgetRow()
+		.Visibility(MediaSourceAssetVisibility)
+		.NameContent()
+		[
+			SNew(STextBlock)
+				.Text(LOCTEXT("MediaAsset", "Media Source Asset"))
+				.ToolTipText(LOCTEXT("MediaAsset_ToolTip", "The Media Source to play."))
+				.Font(IDetailLayoutBuilder::GetDetailFont())
 		]
 		.ValueContent()
 		[
@@ -152,12 +189,12 @@ void FMediaPlateCustomization::CustomizeDetails(IDetailLayoutBuilder& DetailBuil
 	// Add media path.
 	FString FileTypeFilter = TEXT("All files (*.*)|*.*");
 	PlaylistGroup.AddWidgetRow()
+		.Visibility(MediaSourceFileVisibility)
 		.NameContent()
 		[
 			SNew(STextBlock)
 				.Text(LOCTEXT("MediaPath", "Media Path"))
-				.ToolTipText(LOCTEXT("MediaPath_ToolTip",
-					"The path of the Media Source to play.\nChanging this will create a new media source in the level to play this path."))
+				.ToolTipText(LOCTEXT("MediaPath_ToolTip", "The path to a media file like a .mp4."))
 				.Font(IDetailLayoutBuilder::GetDetailFont())
 		]
 		.ValueContent()
@@ -890,6 +927,72 @@ void FMediaPlateCustomization::SetMeshRange(FVector2D Range)
 	}
 }
 
+void FMediaPlateCustomization::UpdateIsMediaSourceAsset()
+{
+	bIsMediaSourceAsset = false;
+
+	// Get the first media plate.
+	if (MediaPlatesList.Num() > 0)
+	{
+		UMediaPlateComponent* MediaPlate = MediaPlatesList[0].Get();
+		if (MediaPlate != nullptr)
+		{
+			UMediaPlaylist* Playlist = MediaPlate->MediaPlaylist;
+			if (Playlist != nullptr)
+			{
+				// Get the first media source in the playlist.
+				UMediaSource* MediaSource = Playlist->Get(0);
+				if (MediaSource != nullptr)
+				{
+					bIsMediaSourceAsset = MediaSource->GetOuter() != MediaPlate;
+				}
+			}
+		}
+	}
+}
+
+void FMediaPlateCustomization::SetIsMediaSourceAsset(bool bIsAsset)
+{
+	if (bIsMediaSourceAsset != bIsAsset)
+	{
+		bIsMediaSourceAsset = bIsAsset;
+
+		// Clear out asset.
+		for (TWeakObjectPtr<UMediaPlateComponent>& MediaPlatePtr : MediaPlatesList)
+		{
+			UMediaPlateComponent* MediaPlate = MediaPlatePtr.Get();
+			if (MediaPlate != nullptr)
+			{
+				// Get playlist.
+				UMediaPlaylist* Playlist = MediaPlate->MediaPlaylist;
+				if (Playlist != nullptr)
+				{
+					// Update playlist.
+					if (Playlist->Num() > 0)
+					{
+						Playlist->Replace(0, nullptr);
+						Playlist->MarkPackageDirty();
+					}
+				}
+				break;
+			}
+		}
+
+		StopMediaPlates();
+		UpdateMediaPath();
+	}
+}
+
+EVisibility FMediaPlateCustomization::ShouldShowMediaSourceAsset() const
+{
+	return bIsMediaSourceAsset ? EVisibility::Visible : EVisibility::Hidden;
+}
+
+EVisibility FMediaPlateCustomization::ShouldShowMediaSourceFile() const
+{
+	return bIsMediaSourceAsset ? EVisibility::Hidden : EVisibility::Visible;
+}
+
 FString FMediaPlateCustomization::GetMediaSourcePath() const
 {
 	FString Path;
@@ -955,6 +1058,7 @@ void FMediaPlateCustomization::OnPlaylistChanged(const FAssetData& AssetData)
 		}
 	}
 
+	UpdateIsMediaSourceAsset();
 	UpdateMediaPath();
 }
 
diff --git a/Engine/Plugins/Media/MediaPlate/Source/MediaPlateEditor/Private/MediaPlateCustomization.h b/Engine/Plugins/Media/MediaPlate/Source/MediaPlateEditor/Private/MediaPlateCustomization.h
--- a/Engine/Plugins/Media/MediaPlate/Source/MediaPlateEditor/Private/MediaPlateCustomization.h
+++ b/Engine/Plugins/Media/MediaPlate/Source/MediaPlateEditor/Private/MediaPlateCustomization.h
@@ -56,6 +56,9 @@ private:
 	/** Handles mesh stuff. */
 	FMediaPlateCustomizationMesh MeshCustomization;
 
+	/** True if the first media source in the playlist is an external asset. */
+	bool bIsMediaSourceAsset;
+
 	/**
 	 * Object property change callback used for static mesh material changes.
 	 */
@@ -156,6 +159,26 @@ private:
 	 */
 	void SetMeshRange(FVector2D Range);
 
+	/**
+	 * Updates bIsMediaSourceAsset depending on the current playlist.
+	 */
+	void UpdateIsMediaSourceAsset();
+
+	/**
+	 * Call this to enable/disable automatic aspect ratio.
+	 */
+	void SetIsMediaSourceAsset(bool bIsAsset);
+
+	/**
+	 * Controls visibility for widgets when the media source is an asset.
+	 */
+	EVisibility ShouldShowMediaSourceAsset() const;
+
+	/**
+	 * Controls visibility for widgets when the media source is a file.
+	 */
+	EVisibility ShouldShowMediaSourceFile() const;
+
 	/**
 	 * Gets the object path for the media source object.
 	 */
