MediaPlate: Fixed apply material not being called when the media plate is not in the details pane.

#rnx

diff --git a/Engine/Plugins/Media/MediaPlate/Source/MediaPlate/MediaPlate.Build.cs b/Engine/Plugins/Media/MediaPlate/Source/MediaPlate/MediaPlate.Build.cs
--- a/Engine/Plugins/Media/MediaPlate/Source/MediaPlate/MediaPlate.Build.cs
+++ b/Engine/Plugins/Media/MediaPlate/Source/MediaPlate/MediaPlate.Build.cs
@@ -25,5 +25,10 @@ public class MediaPlate : ModuleRules
 				"MediaUtils",
 			}
 			);
+
+		if (Target.bBuildEditor)
+		{
+			PrivateDependencyModuleNames.Add("UnrealEd");
+		}
 	}
 }
diff --git a/Engine/Plugins/Media/MediaPlate/Source/MediaPlate/Private/MediaPlate.cpp b/Engine/Plugins/Media/MediaPlate/Source/MediaPlate/Private/MediaPlate.cpp
--- a/Engine/Plugins/Media/MediaPlate/Source/MediaPlate/Private/MediaPlate.cpp
+++ b/Engine/Plugins/Media/MediaPlate/Source/MediaPlate/Private/MediaPlate.cpp
@@ -11,6 +11,12 @@
 #include "MediaTexture.h"
 #include "UObject/ConstructorHelpers.h"
 
+#if WITH_EDITOR
+#include "Editor.h"
+#include "MediaPlateAssetUserData.h"
+#include "UObject/ObjectSaveContext.h"
+#endif
+
 #define LOCTEXT_NAMESPACE "MediaPlate"
 
 FLazyName AMediaPlate::MediaPlateComponentName(TEXT("MediaPlateComponent0"));
@@ -45,8 +51,17 @@ AMediaPlate::AMediaPlate(const FObjectInitializer& ObjectInitializer)
 	{
 		StaticMeshComponent->SetStaticMesh(ConstructorStatics.Plane.Object);
 	}
-}
 
+#if WITH_EDITOR
+	// If we are not the class default object...
+	if (!HasAnyFlags(RF_ClassDefaultObject))
+	{
+		// Hook into pre/post save.
+		FEditorDelegates::PreSaveWorldWithContext.AddUObject(this, &AMediaPlate::OnPreSaveWorld);
+		FEditorDelegates::PostSaveWorldWithContext.AddUObject(this, &AMediaPlate::OnPostSaveWorld);
+	}
+#endif // WITH_EDITOR
+}
 
 void AMediaPlate::PostActorCreated()
 {
@@ -70,6 +85,10 @@ void AMediaPlate::PostRegisterAllComponents()
 		// Add our media texture to the tracker.
 		MediaPlateComponent->RegisterWithMediaTextureTracker();
 	}
+
+#if WITH_EDITOR
+	AddAssetUserData();
+#endif // WITH_EDITOR
 }
 
 void AMediaPlate::BeginDestroy()
@@ -92,6 +111,18 @@ void AMediaPlate::UseDefaultMaterial()
 	ApplyMaterial(DefaultMaterial);
 }
 
+void AMediaPlate::ApplyCurrentMaterial()
+{
+	if (StaticMeshComponent != nullptr)
+	{
+		UMaterialInterface* MaterialInterface = StaticMeshComponent->GetMaterial(0);
+		if ((MaterialInterface != nullptr) && (LastMaterial != MaterialInterface))
+		{
+			ApplyMaterial(MaterialInterface);
+		}
+	}
+}
+
 void AMediaPlate::ApplyMaterial(UMaterialInterface* Material)
 {
 	if (Material != nullptr && StaticMeshComponent != nullptr)
@@ -143,6 +174,36 @@ void AMediaPlate::ApplyMaterial(UMaterialInterface* Material)
 	}
 }
 
+void AMediaPlate::OnPreSaveWorld(UWorld* InWorld, FObjectPreSaveContext ObjectSaveContext)
+{
+	// We need to remove our asset user data before saving, as we do not need to save it out
+	// and only use it to know when the static mesh component changes.
+	RemoveAssetUserData();
+}
+
+void AMediaPlate::OnPostSaveWorld(UWorld* InWorld, FObjectPostSaveContext ObjectSaveContext)
+{
+	AddAssetUserData();
+}
+
+void AMediaPlate::AddAssetUserData()
+{
+	if (StaticMeshComponent != nullptr)
+	{
+		UMediaPlateAssetUserData* AssetUserData = NewObject<UMediaPlateAssetUserData>(GetTransientPackage());
+		AssetUserData->OnPostEditChangeOwner.BindUObject(this, &AMediaPlate::ApplyCurrentMaterial);
+		StaticMeshComponent->AddAssetUserData(AssetUserData);
+	}
+}
+
+void AMediaPlate::RemoveAssetUserData()
+{
+	if (StaticMeshComponent != nullptr)
+	{
+		StaticMeshComponent->RemoveUserDataOfClass(UMediaPlateAssetUserData::StaticClass());
+	}
+}
+
 #endif
 
 #undef LOCTEXT_NAMESPACE
diff --git a/Engine/Plugins/Media/MediaPlate/Source/MediaPlate/Private/MediaPlateAssetUserData.cpp b/Engine/Plugins/Media/MediaPlate/Source/MediaPlate/Private/MediaPlateAssetUserData.cpp
new file mode 100644
--- /dev/null
+++ b/Engine/Plugins/Media/MediaPlate/Source/MediaPlate/Private/MediaPlateAssetUserData.cpp
@@ -0,0 +1,8 @@
+// Copyright Epic Games, Inc. All Rights Reserved.
+
+#include "MediaPlateAssetUserData.h"
+
+void UMediaPlateAssetUserData::PostEditChangeOwner()
+{
+	OnPostEditChangeOwner.ExecuteIfBound();
+}
diff --git a/Engine/Plugins/Media/MediaPlate/Source/MediaPlate/Public/MediaPlate.h b/Engine/Plugins/Media/MediaPlate/Source/MediaPlate/Public/MediaPlate.h
--- a/Engine/Plugins/Media/MediaPlate/Source/MediaPlate/Public/MediaPlate.h
+++ b/Engine/Plugins/Media/MediaPlate/Source/MediaPlate/Public/MediaPlate.h
@@ -4,6 +4,7 @@
 
 #include "CoreMinimal.h"
 #include "GameFramework/Actor.h"
+#include "Engine/AssetUserData.h"
 #include "MediaPlate.generated.h"
 
 class UMediaPlateComponent;
@@ -36,9 +37,14 @@ public:
 	 * Call this to change the static mesh to use the default media plate material.
 	 */
 	void UseDefaultMaterial();
+
+	/**
+	 * Call this after changing the current material to set it up for media plate.
+	 */
+	void ApplyCurrentMaterial();
 	void ApplyMaterial(UMaterialInterface* InMaterial);
 	UMaterialInterface* GetLastMaterial() const { return LastMaterial; }
-#endif
+#endif // WITH_EDITOR
 
 private:
 	/** Name for our media plate component. */
@@ -48,5 +54,26 @@ private:
 
 #if WITH_EDITOR
 	UMaterialInterface* LastMaterial = nullptr;
-#endif
+
+	/**
+	 * Called before a level saves
+	 */
+	void OnPreSaveWorld(UWorld* InWorld, FObjectPreSaveContext ObjectSaveContext);
+
+	/**
+	 * Called after a level has saved.
+	 */
+	void OnPostSaveWorld(UWorld* InWorld, FObjectPostSaveContext ObjectSaveContext);
+
+	/**
+	 * Adds our asset user data to the static mesh component.
+	 */
+	void AddAssetUserData();
+
+	/**
+	 * Removes our asset user data from the static mesh component.
+	 */
+	void RemoveAssetUserData();
+
+#endif // WITH_EDITOR
 };
diff --git a/Engine/Plugins/Media/MediaPlate/Source/MediaPlate/Public/MediaPlateAssetUserData.h b/Engine/Plugins/Media/MediaPlate/Source/MediaPlate/Public/MediaPlateAssetUserData.h
new file mode 100644
--- /dev/null
+++ b/Engine/Plugins/Media/MediaPlate/Source/MediaPlate/Public/MediaPlateAssetUserData.h
@@ -0,0 +1,25 @@
+// Copyright Epic Games, Inc. All Rights Reserved.
+
+#pragma once
+
+#include "Engine/AssetUserData.h"
+#include "MediaPlateAssetUserData.generated.h"
+
+DECLARE_DELEGATE(FOnPostEditChangeOwner);
+
+/**
+ * AssetUserData for media plate.
+ */
+UCLASS()
+class UMediaPlateAssetUserData : public UAssetUserData
+{
+	GENERATED_BODY()
+
+public:
+	/** Bind to this to get a callback when something is changed on the associated component. */
+	FOnPostEditChangeOwner OnPostEditChangeOwner;
+	
+	//~ Begin UAssetUserData
+	virtual void PostEditChangeOwner() override;
+	//~ End UAssetUserData
+};
diff --git a/Engine/Plugins/Media/MediaPlate/Source/MediaPlateEditor/Private/MediaPlateCustomization.cpp b/Engine/Plugins/Media/MediaPlate/Source/MediaPlateEditor/Private/MediaPlateCustomization.cpp
--- a/Engine/Plugins/Media/MediaPlate/Source/MediaPlateEditor/Private/MediaPlateCustomization.cpp
+++ b/Engine/Plugins/Media/MediaPlate/Source/MediaPlateEditor/Private/MediaPlateCustomization.cpp
@@ -34,40 +34,10 @@
 FMediaPlateCustomization::FMediaPlateCustomization()
 	: bIsMediaSourceAsset(false)
 {
-	PropertyChangeDelegate = FCoreUObjectDelegates::OnObjectPropertyChanged.AddRaw(this, &FMediaPlateCustomization::OnObjectPropertyChanged);
 }
 
 FMediaPlateCustomization::~FMediaPlateCustomization()
 {
-	FCoreUObjectDelegates::OnObjectPropertyChanged.RemoveAll(this);
-}
-
-void FMediaPlateCustomization::OnObjectPropertyChanged(UObject* InObject, FPropertyChangedEvent& InEvent)
-{
-	//Note: Cannot check for GET_MEMBER_NAME_CHECKED(UMeshComponent, OverrideMaterials) because the event is null on drag-and-drop.
-
-	for (TWeakObjectPtr<UMediaPlateComponent>& MediaPlateComponent : MediaPlatesList)
-	{
-		if (MediaPlateComponent.IsValid())
-		{
-			if (InObject->GetOuter() == MediaPlateComponent->GetOuter())
-			{
-				AMediaPlate* MediaPlate = Cast<AMediaPlate>(InObject->GetOuter());
-
-				if (MediaPlate != nullptr && MediaPlate->StaticMeshComponent != nullptr)
-				{
-					UMaterialInterface* Material = MediaPlate->StaticMeshComponent->GetMaterial(0);
-
-					if (Material != MediaPlate->GetLastMaterial())
-					{
-						MediaPlate->ApplyMaterial(Material);
-					}
-				}
-
-				break;
-			}
-		}
-	}
 }
 
 void FMediaPlateCustomization::CustomizeDetails(IDetailLayoutBuilder& DetailBuilder)
diff --git a/Engine/Plugins/Media/MediaPlate/Source/MediaPlateEditor/Private/MediaPlateCustomization.h b/Engine/Plugins/Media/MediaPlate/Source/MediaPlateEditor/Private/MediaPlateCustomization.h
--- a/Engine/Plugins/Media/MediaPlate/Source/MediaPlateEditor/Private/MediaPlateCustomization.h
+++ b/Engine/Plugins/Media/MediaPlate/Source/MediaPlateEditor/Private/MediaPlateCustomization.h
@@ -59,11 +59,6 @@ private:
 	/** True if the first media source in the playlist is an external asset. */
 	bool bIsMediaSourceAsset;
 
-	/**
-	 * Object property change callback used for static mesh material changes.
-	 */
-	void OnObjectPropertyChanged(UObject* InObject, FPropertyChangedEvent& InEvent);
-
 	/**
 	 * Adds widgets for editing the mesh.
 	 */
