MediaPlate: Static mesh moved to media plate customization area.

#rnx

diff --git a/Engine/Plugins/Media/MediaPlate/Source/MediaPlateEditor/Private/MediaPlateCustomization.cpp b/Engine/Plugins/Media/MediaPlate/Source/MediaPlateEditor/Private/MediaPlateCustomization.cpp
--- a/Engine/Plugins/Media/MediaPlate/Source/MediaPlateEditor/Private/MediaPlateCustomization.cpp
+++ b/Engine/Plugins/Media/MediaPlate/Source/MediaPlateEditor/Private/MediaPlateCustomization.cpp
@@ -17,6 +17,7 @@
 #include "PropertyCustomizationHelpers.h"
 #include "Styling/AppStyle.h"
 #include "Subsystems/AssetEditorSubsystem.h"
+#include "ThumbnailRendering/ThumbnailManager.h"
 #include "Widgets/Input/SButton.h"
 #include "Widgets/Input/SCheckBox.h"
 #include "Widgets/Input/SComboButton.h"
@@ -59,6 +60,10 @@ void FMediaPlateCustomization::CustomizeDetails(IDetailLayoutBuilder& DetailBuil
 		Style = &FAppStyle::Get();
 	}
 
+	// Hide the static mesh.
+	IDetailCategoryBuilder& StaticMeshCategory = DetailBuilder.EditCategory("StaticMesh");
+	StaticMeshCategory.SetCategoryVisibility(false);
+
 	IDetailCategoryBuilder& MediaPlateCategory = DetailBuilder.EditCategory("MediaPlate");
 
 	// Get objects we are editing.
@@ -489,6 +494,7 @@ void FMediaPlateCustomization::AddMeshCustomization(IDetailCategoryBuilder& Medi
 		];
 
 	// Visibility attributes.
+	TAttribute<EVisibility> MeshCustomVisibility(this, &FMediaPlateCustomization::ShouldShowMeshCustomWidgets);
 	TAttribute<EVisibility> MeshPlaneVisibility(this, &FMediaPlateCustomization::ShouldShowMeshPlaneWidgets);
 	TAttribute<EVisibility> MeshSphereVisibility(this, &FMediaPlateCustomization::ShouldShowMeshSphereWidgets);
 
@@ -628,9 +634,33 @@ void FMediaPlateCustomization::AddMeshCustomization(IDetailCategoryBuilder& Medi
 				.Value(this, &FMediaPlateCustomization::GetMeshVerticalRange)
 				.OnValueChanged(this, &FMediaPlateCustomization::SetMeshVerticalRange)
 		];
+
+	// Add static mesh.
+	DetailGroup.AddWidgetRow()
+		.Visibility(MeshCustomVisibility)
+		.NameContent()
+		[
+			SNew(STextBlock)
+				.Text(LOCTEXT("StaticMesh", "Static Mesh"))
+				.ToolTipText(LOCTEXT("StaticMesh_Tooltip", "The static mesh to use."))
+				.Font(IDetailLayoutBuilder::GetDetailFont())
+		]
+		.ValueContent()
+		[
+			SNew(SObjectPropertyEntryBox)
+				.AllowedClass(UStaticMesh::StaticClass())
+				.ThumbnailPool(UThumbnailManager::Get().GetSharedThumbnailPool())
+				.ObjectPath(this, &FMediaPlateCustomization::GetStaticMeshPath)
+				.OnObjectChanged(this, &FMediaPlateCustomization::OnStaticMeshChanged)
+		];
 	
 }
 
+EVisibility FMediaPlateCustomization::ShouldShowMeshCustomWidgets() const
+{
+	return (MeshMode == EMediaTextureVisibleMipsTiles::None) ? EVisibility::Visible : EVisibility::Hidden;
+}
+
 EVisibility FMediaPlateCustomization::ShouldShowMeshPlaneWidgets() const
 {
 	return (MeshMode == EMediaTextureVisibleMipsTiles::Plane) ? EVisibility::Visible : EVisibility::Hidden;
@@ -896,6 +926,46 @@ void FMediaPlateCustomization::SetMeshRange(FVector2D Range)
 	}
 }
 
+
+FString FMediaPlateCustomization::GetStaticMeshPath() const
+{
+	FString Path;
+
+	// Get the first media plate.
+	if (MediaPlatesList.Num() > 0)
+	{
+		UMediaPlateComponent* MediaPlate = MediaPlatesList[0].Get();
+		if (MediaPlate != nullptr)
+		{
+			UStaticMeshComponent* StaticMeshComponent = MediaPlate->StaticMeshComponent;
+			if (StaticMeshComponent != nullptr)
+			{
+				UStaticMesh* StaticMesh = StaticMeshComponent->GetStaticMesh();
+				if (StaticMesh != nullptr)
+				{
+					Path = StaticMesh->GetPathName();
+				}
+			}
+		}
+	}
+
+	return Path;
+}
+
+void FMediaPlateCustomization::OnStaticMeshChanged(const FAssetData& AssetData)
+{
+	// Update the static mesh.
+	UStaticMesh* StaticMesh = Cast<UStaticMesh>(AssetData.GetAsset());
+	for (TWeakObjectPtr<UMediaPlateComponent>& MediaPlatePtr : MediaPlatesList)
+	{
+		UMediaPlateComponent* MediaPlate = MediaPlatePtr.Get();
+		if (MediaPlate != nullptr)
+		{
+			MeshCustomization.SetCustomMesh(MediaPlate, StaticMesh);
+		}
+	}
+}
+
 void FMediaPlateCustomization::UpdateIsMediaSourceAsset()
 {
 	bIsMediaSourceAsset = false;
diff --git a/Engine/Plugins/Media/MediaPlate/Source/MediaPlateEditor/Private/MediaPlateCustomization.h b/Engine/Plugins/Media/MediaPlate/Source/MediaPlateEditor/Private/MediaPlateCustomization.h
--- a/Engine/Plugins/Media/MediaPlate/Source/MediaPlateEditor/Private/MediaPlateCustomization.h
+++ b/Engine/Plugins/Media/MediaPlate/Source/MediaPlateEditor/Private/MediaPlateCustomization.h
@@ -64,6 +64,11 @@ private:
 	 */
 	void AddMeshCustomization(IDetailCategoryBuilder& MediaPlateCategory);
 
+	/**
+	 * Controls visibility for widgets for custom meshes.
+	 */
+	EVisibility ShouldShowMeshCustomWidgets() const;
+
 	/**
 	 * Controls visibility for widgets for plane meshes.
 	 */
@@ -154,6 +159,16 @@ private:
 	 */
 	void SetMeshRange(FVector2D Range);
 
+	/**
+	 * Gets the object path for the static mesh.
+	 */
+	FString GetStaticMeshPath() const;
+
+	/**
+	 * Called when the static mesh changes.
+	 */
+	void OnStaticMeshChanged(const FAssetData& AssetData);
+
 	/**
 	 * Updates bIsMediaSourceAsset depending on the current playlist.
 	 */
diff --git a/Engine/Plugins/Media/MediaPlate/Source/MediaPlateEditor/Private/MediaPlateCustomizationMesh.cpp b/Engine/Plugins/Media/MediaPlate/Source/MediaPlateEditor/Private/MediaPlateCustomizationMesh.cpp
--- a/Engine/Plugins/Media/MediaPlate/Source/MediaPlateEditor/Private/MediaPlateCustomizationMesh.cpp
+++ b/Engine/Plugins/Media/MediaPlate/Source/MediaPlateEditor/Private/MediaPlateCustomizationMesh.cpp
@@ -14,7 +14,8 @@
 
 TMap<UStaticMesh*, int32> FMediaPlateCustomizationMesh::MeshRefCount;
 
-void FMediaPlateCustomizationMesh::SetPlaneMesh(UMediaPlateComponent* MediaPlate)
+void FMediaPlateCustomizationMesh::SetCustomMesh(UMediaPlateComponent* MediaPlate,
+	UStaticMesh* StaticMesh)
 {
 	// Get static mesh component.
 	UStaticMeshComponent* StaticMeshComponent = nullptr;
@@ -30,15 +31,21 @@ void FMediaPlateCustomizationMesh::SetPlaneMesh(UMediaPlateComponent* MediaPlate
 
 	if (StaticMeshComponent != nullptr)
 	{
-		// Get plane mesh.
-		UStaticMesh* Mesh = Cast<UStaticMesh>(StaticLoadObject(UStaticMesh::StaticClass(), NULL,
-			TEXT("/MediaPlate/SM_MediaPlateScreen")));
-
 		// Apply this mesh.
-		SetMesh(StaticMeshComponent, Mesh);
+		SetMesh(StaticMeshComponent, StaticMesh);
 	}
 }
 
+void FMediaPlateCustomizationMesh::SetPlaneMesh(UMediaPlateComponent* MediaPlate)
+{
+	// Get plane mesh.
+	UStaticMesh* Mesh = Cast<UStaticMesh>(StaticLoadObject(UStaticMesh::StaticClass(), NULL,
+		TEXT("/MediaPlate/SM_MediaPlateScreen")));
+
+	// Apply this mesh.
+	SetCustomMesh(MediaPlate, Mesh);
+}
+
 void FMediaPlateCustomizationMesh::SetSphereMesh(UMediaPlateComponent* MediaPlate)
 {
 	// Get static mesh component.
diff --git a/Engine/Plugins/Media/MediaPlate/Source/MediaPlateEditor/Private/MediaPlateCustomizationMesh.h b/Engine/Plugins/Media/MediaPlate/Source/MediaPlateEditor/Private/MediaPlateCustomizationMesh.h
--- a/Engine/Plugins/Media/MediaPlate/Source/MediaPlateEditor/Private/MediaPlateCustomizationMesh.h
+++ b/Engine/Plugins/Media/MediaPlate/Source/MediaPlateEditor/Private/MediaPlateCustomizationMesh.h
@@ -16,6 +16,11 @@ class UStaticMeshComponent;
 class FMediaPlateCustomizationMesh
 {
 public:
+	/**
+	 * Call this to use a custom mesh on a media plate.
+	 */
+	void SetCustomMesh(UMediaPlateComponent* MediaPlate, UStaticMesh* StaticMesh);
+
 	/**
 	 * Call this to use a plane mesh on a media plate.
 	 */
