MediaPlate: Removed unused style for SMediaPlateEditorMediaDetails.
#rnx

diff --git a/Engine/Plugins/Media/MediaPlate/Source/MediaPlateEditor/Private/Toolkits/MediaPlateEditorToolkit.cpp b/Engine/Plugins/Media/MediaPlate/Source/MediaPlateEditor/Private/Toolkits/MediaPlateEditorToolkit.cpp
--- a/Engine/Plugins/Media/MediaPlate/Source/MediaPlateEditor/Private/Toolkits/MediaPlateEditorToolkit.cpp
+++ b/Engine/Plugins/Media/MediaPlate/Source/MediaPlateEditor/Private/Toolkits/MediaPlateEditorToolkit.cpp
@@ -397,7 +397,7 @@ TSharedRef<SDockTab> FMediaPlateEditorToolkit::HandleTabManagerSpawnTab(const FS
 	}
 	else if (TabIdentifier == MediaPlateEditorToolkit::MediaDetailsTabId)
 	{
-		TabWidget = SNew(SMediaPlateEditorMediaDetails, *MediaPlate, Style);
+		TabWidget = SNew(SMediaPlateEditorMediaDetails, *MediaPlate);
 	}
 	else if (TabIdentifier == MediaPlateEditorToolkit::PlaylistTabId)
 	{
diff --git a/Engine/Plugins/Media/MediaPlate/Source/MediaPlateEditor/Private/Widgets/SMediaPlateEditorMediaDetails.cpp b/Engine/Plugins/Media/MediaPlate/Source/MediaPlateEditor/Private/Widgets/SMediaPlateEditorMediaDetails.cpp
--- a/Engine/Plugins/Media/MediaPlate/Source/MediaPlateEditor/Private/Widgets/SMediaPlateEditorMediaDetails.cpp
+++ b/Engine/Plugins/Media/MediaPlate/Source/MediaPlateEditor/Private/Widgets/SMediaPlateEditorMediaDetails.cpp
@@ -15,7 +15,8 @@
 /* SMediaPlateEditorMediaDetails interface
  *****************************************************************************/
 
-void SMediaPlateEditorMediaDetails::Construct(const FArguments& InArgs, UMediaPlateComponent& InMediaPlate, const TSharedRef<ISlateStyle>& InStyle)
+void SMediaPlateEditorMediaDetails::Construct(const FArguments& InArgs,
+	UMediaPlateComponent& InMediaPlate)
 {
 	MediaPlate = &InMediaPlate;
 
diff --git a/Engine/Plugins/Media/MediaPlate/Source/MediaPlateEditor/Private/Widgets/SMediaPlateEditorMediaDetails.h b/Engine/Plugins/Media/MediaPlate/Source/MediaPlateEditor/Private/Widgets/SMediaPlateEditorMediaDetails.h
--- a/Engine/Plugins/Media/MediaPlate/Source/MediaPlateEditor/Private/Widgets/SMediaPlateEditorMediaDetails.h
+++ b/Engine/Plugins/Media/MediaPlate/Source/MediaPlateEditor/Private/Widgets/SMediaPlateEditorMediaDetails.h
@@ -28,9 +28,8 @@ public:
 	 *
 	 * @param InArgs		The declaration data for this widget.
 	 * @param InMediaPlate	The MediaPlate to show the details for.
-	 * @param InStyleSet	The style set to use.
 	 */
-	void Construct(const FArguments& InArgs, UMediaPlateComponent& InMediaPlate, const TSharedRef<ISlateStyle>& InStyle);
+	void Construct(const FArguments& InArgs, UMediaPlateComponent& InMediaPlate);
 
 	//~ SWidget interface
 	virtual void Tick(const FGeometry& AllottedGeometry, const double InCurrentTime, const float InDeltaTime) override;
