MediaPlate: Added media details to customization.

#rnx

diff --git a/Engine/Plugins/Media/MediaPlate/Source/MediaPlateEditor/Private/MediaPlateCustomization.cpp b/Engine/Plugins/Media/MediaPlate/Source/MediaPlateEditor/Private/MediaPlateCustomization.cpp
--- a/Engine/Plugins/Media/MediaPlate/Source/MediaPlateEditor/Private/MediaPlateCustomization.cpp
+++ b/Engine/Plugins/Media/MediaPlate/Source/MediaPlateEditor/Private/MediaPlateCustomization.cpp
@@ -25,6 +25,7 @@
 #include "Widgets/Input/SNumericEntryBox.h"
 #include "Widgets/Input/SSegmentedControl.h"
 #include "Widgets/Images/SImage.h"
+#include "Widgets/SMediaPlateEditorMediaDetails.h"
 
 #define LOCTEXT_NAMESPACE "FMediaPlateCustomization"
 
@@ -458,6 +459,30 @@ void FMediaPlateCustomization::CustomizeDetails(IDetailLayoutBuilder& DetailBuil
 							.Text(LOCTEXT("OpenMediaPlate", "Open Media Plate"))
 					]
 			];
+
+		// Get the first media plate.
+		UMediaPlateComponent* FirstMediaPlate = nullptr;
+		for (TWeakObjectPtr<UMediaPlateComponent>& MediaPlatePtr : MediaPlatesList)
+		{
+			UMediaPlateComponent* MediaPlate = MediaPlatePtr.Get();
+			if (MediaPlate != nullptr)
+			{
+				FirstMediaPlate = MediaPlate;
+				break;
+			}
+		}
+
+		if (FirstMediaPlate != nullptr)
+		{
+			// Create media details group.
+			IDetailGroup& MediaDetailsGroup = MediaPlateCategory.AddGroup(TEXT("MediaDetails"),
+				LOCTEXT("MediaDetails", "Media Details"));
+
+			MediaDetailsGroup.AddWidgetRow()
+			[
+				SNew(SMediaPlateEditorMediaDetails, *FirstMediaPlate)
+			];
+		}
 	}
 }
 
