MediaPlate: Fixed PlayOnlyWhenVisible visibility checks not working in editor.
#rnx

diff --git a/Engine/Plugins/Media/MediaPlate/Source/MediaPlate/Private/MediaPlateComponent.cpp b/Engine/Plugins/Media/MediaPlate/Source/MediaPlate/Private/MediaPlateComponent.cpp
--- a/Engine/Plugins/Media/MediaPlate/Source/MediaPlate/Private/MediaPlateComponent.cpp
+++ b/Engine/Plugins/Media/MediaPlate/Source/MediaPlate/Private/MediaPlateComponent.cpp
@@ -63,6 +63,7 @@ FLazyName UMediaPlateComponent::MediaPlaylistName(TEXT("MediaPlaylist0"));
 UMediaPlateComponent::UMediaPlateComponent(const FObjectInitializer& ObjectInitializer)
 	: Super(ObjectInitializer)
 {
+	bTickInEditor = true;
 	PrimaryComponentTick.bCanEverTick = true;
 	PrimaryComponentTick.bStartWithTickEnabled = false;
 
@@ -111,15 +112,15 @@ void UMediaPlateComponent::OnRegister()
 			SoundComponent->SetMediaPlayer(MediaPlayer);
 		}
 	}
+
+	// Activate tick if needed.
+	UpdateTicking();
 }
 
 void UMediaPlateComponent::BeginPlay()
 {
 	Super::BeginPlay();
 
-	// Activate tick if needed.
-	UpdateTicking();
-
 	// Start playing?
 	if (bAutoPlay)
 	{
@@ -467,15 +468,6 @@ void UMediaPlateComponent::StopClockSink()
 
 bool UMediaPlateComponent::IsVisible()
 {
-#if WITH_EDITOR
-	// Always return true if the game is not running as the visibility checks do not work in this
-	// case.
-	if (FApp::IsGame() == false)
-	{
-		return true;
-	}
-#endif
-
 	return GetOwner()->WasRecentlyRendered();
 }
 
diff --git a/Engine/Plugins/Media/MediaPlate/Source/MediaPlate/Public/MediaPlateComponent.h b/Engine/Plugins/Media/MediaPlate/Source/MediaPlate/Public/MediaPlateComponent.h
--- a/Engine/Plugins/Media/MediaPlate/Source/MediaPlate/Public/MediaPlateComponent.h
+++ b/Engine/Plugins/Media/MediaPlate/Source/MediaPlate/Public/MediaPlateComponent.h
@@ -172,7 +172,7 @@ public:
 	void UnregisterWithMediaTextureTracker();
 
 	/**
-	 * Call this to geet the aspect ratio of the mesh.
+	 * Call this to get the aspect ratio of the mesh.
 	 */
 	float GetAspectRatio();
 
@@ -193,6 +193,11 @@ public:
 	UFUNCTION(BlueprintCallable, Category = "Media|MediaPlateComponent")
 	void SetLetterboxAspectRatio(float AspectRatio);
 
+	/**
+	 * Call this to see if this plate wants to play when it becomes visible.
+	 */
+	bool GetWantsToPlayWhenVisible() const { return bWantsToPlayWhenVisible; }
+
 #if WITH_EDITOR
 	/**
 	 * Call this to get the mip tile calculations mesh mode.
diff --git a/Engine/Plugins/Media/MediaPlate/Source/MediaPlateEditor/Private/MediaPlateEditorModule.cpp b/Engine/Plugins/Media/MediaPlate/Source/MediaPlateEditor/Private/MediaPlateEditorModule.cpp
--- a/Engine/Plugins/Media/MediaPlate/Source/MediaPlateEditor/Private/MediaPlateEditorModule.cpp
+++ b/Engine/Plugins/Media/MediaPlate/Source/MediaPlateEditor/Private/MediaPlateEditorModule.cpp
@@ -5,7 +5,7 @@
 #include "AssetTools/MediaPlateActions.h"
 #include "Editor.h"
 #include "ISequencerModule.h"
-#include "LevelEditorViewport.h"
+#include "LevelEditor.h"
 #include "MaterialList.h"
 #include "MediaPlateComponent.h"
 #include "MediaPlateCustomization.h"
@@ -14,12 +14,15 @@
 #include "MediaSoundComponent.h"
 #include "MediaSource.h"
 #include "Models/MediaPlateEditorCommands.h"
+#include "Modules/ModuleManager.h"
 #include "PropertyEditorModule.h"
 #include "Sequencer/MediaPlateTrackEditor.h"
 #include "SLevelViewport.h"
 #include "Subsystems/EditorAssetSubsystem.h"
 #include "Widgets/SMediaPlateEditorMaterial.h"
 
+#define LOCTEXT_NAMESPACE "MediaPlateEditorModule"
+
 DEFINE_LOG_CATEGORY(LogMediaPlateEditor);
 
 void FMediaPlateEditorModule::StartupModule()
@@ -102,7 +105,7 @@ void FMediaPlateEditorModule::Tick(float DeltaTime)
 			{
 				bIsMediaPlateToBeRemoved = false;
 				// Are we playing something?
-				if (MediaPlayer->IsPlaying())
+				if ((MediaPlayer->IsPlaying()) || (MediaPlate->GetWantsToPlayWhenVisible()))
 				{
 					bIsMediaPlatePlaying = true;
 					break;
@@ -119,14 +122,7 @@ void FMediaPlateEditorModule::Tick(float DeltaTime)
 	}
 
 	// Is anything playing?
-	if (bIsMediaPlatePlaying)
-	{
-		// Yes. Invalidate the viewport so we can see it.
-		if (GCurrentLevelEditingViewportClient != nullptr)
-		{
-			GCurrentLevelEditingViewportClient->Invalidate();
-		}
-	}
+	ForceRealTimeViewports(bIsMediaPlatePlaying);
 }
 
 void FMediaPlateEditorModule::MediaPlateStartedPlayback(TObjectPtr<UMediaPlateComponent> MediaPlate)
@@ -217,4 +213,38 @@ void FMediaPlateEditorModule::ExtractAssetDataFromFiles(const TArray<FString>& F
 	}
 }
 
+void FMediaPlateEditorModule::ForceRealTimeViewports(const bool bEnable)
+{
+	if (bIsRealTimeViewportsEnabled != bEnable)
+	{
+		bIsRealTimeViewportsEnabled = bEnable;
+
+		// Go through viewports.
+		FLevelEditorModule& LevelEditorModule = FModuleManager::GetModuleChecked<FLevelEditorModule>("LevelEditor");
+		TSharedPtr<ILevelEditor> LevelEditor = LevelEditorModule.GetFirstLevelEditor();
+		if (LevelEditor.IsValid())
+		{
+			TArray<TSharedPtr<SLevelViewport>> Viewports = LevelEditor->GetViewports();
+			for (const TSharedPtr<SLevelViewport>& ViewportWindow : Viewports)
+			{
+				if (ViewportWindow.IsValid())
+				{
+					FEditorViewportClient& Viewport = ViewportWindow->GetAssetViewportClient();
+					const FText SystemDisplayName = LOCTEXT("RealtimeOverrideMessage", "Media Plate");
+					if (bEnable)
+					{
+						Viewport.AddRealtimeOverride(bEnable, SystemDisplayName);
+					}
+					else
+					{
+						Viewport.RemoveRealtimeOverride(SystemDisplayName, false);
+					}
+				}
+			}
+		}
+	}
+}
+
 IMPLEMENT_MODULE(FMediaPlateEditorModule, MediaPlateEditor)
+
+#undef LOCTEXT_NAMESPACE
diff --git a/Engine/Plugins/Media/MediaPlate/Source/MediaPlateEditor/Public/MediaPlateEditorModule.h b/Engine/Plugins/Media/MediaPlate/Source/MediaPlateEditor/Public/MediaPlateEditorModule.h
--- a/Engine/Plugins/Media/MediaPlate/Source/MediaPlateEditor/Public/MediaPlateEditorModule.h
+++ b/Engine/Plugins/Media/MediaPlate/Source/MediaPlateEditor/Public/MediaPlateEditorModule.h
@@ -68,6 +68,8 @@ private:
 
 	/** Maps a drag and drop file to a media source we created from it. */
 	TMap<FString, TWeakObjectPtr<UMediaSource>> MapFileToMediaSource;
+	/** Caches if real time viewports are enabled. */
+	bool bIsRealTimeViewportsEnabled = false;
 
 	/**
 	 * Registers all of our asset tools.
@@ -99,4 +101,9 @@ private:
 	 * @param AssetDataArray	Will add the asset here if it gets one.
 	 */
 	void ExtractAssetDataFromFiles(const TArray<FString>& Files, TArray<FAssetData>& AssetDataArray);
+
+	/**
+	 * Enables real time viewports so media plates can update in the editor.
+	 */
+	void ForceRealTimeViewports(const bool bEnable);
 };
