MediaPlate: Cosmetic fix.
#rnx

diff --git a/Engine/Plugins/Media/MediaPlate/Source/MediaPlate/Private/MediaPlateComponent.cpp b/Engine/Plugins/Media/MediaPlate/Source/MediaPlate/Private/MediaPlateComponent.cpp
--- a/Engine/Plugins/Media/MediaPlate/Source/MediaPlate/Private/MediaPlateComponent.cpp
+++ b/Engine/Plugins/Media/MediaPlate/Source/MediaPlate/Private/MediaPlateComponent.cpp
@@ -189,15 +189,15 @@ void UMediaPlateComponent::Open()
 	if ((bPlayOnlyWhenVisible == false) || (IsVisible()))
 	{
 		bool bIsPlaying = false;
-			if (MediaPlayer != nullptr)
+		if (MediaPlayer != nullptr)
+		{
+			UMediaSource* MediaSource = nullptr;
+			if (MediaPlaylist != nullptr)
 			{
-				UMediaSource* MediaSource = nullptr;
-					if (MediaPlaylist != nullptr)
-					{
-						MediaSource = MediaPlaylist->Get(0);
-					}
-				bIsPlaying = PlayMediaSource(MediaSource);
+				MediaSource = MediaPlaylist->Get(0);
 			}
+			bIsPlaying = PlayMediaSource(MediaSource);
+		}
 
 		// Did anything play?
 		if (bIsPlaying == false)
