MediaPlate: Fixed play only when visible seeking too early, and always trying to open.
#rnx

diff --git a/Engine/Plugins/Media/MediaPlate/Source/MediaPlate/Private/MediaPlateComponent.cpp b/Engine/Plugins/Media/MediaPlate/Source/MediaPlate/Private/MediaPlateComponent.cpp
--- a/Engine/Plugins/Media/MediaPlate/Source/MediaPlate/Private/MediaPlateComponent.cpp
+++ b/Engine/Plugins/Media/MediaPlate/Source/MediaPlate/Private/MediaPlateComponent.cpp
@@ -94,6 +94,7 @@ void UMediaPlateComponent::OnRegister()
 		MediaPlayer->SetLooping(false);
 		MediaPlayer->PlayOnOpen = false;
 	}
+	MediaPlayer->OnMediaOpened.AddUniqueDynamic(this, &UMediaPlateComponent::OnMediaOpened);
 	MediaPlayer->OnEndReached.AddUniqueDynamic(this, &UMediaPlateComponent::OnMediaEnd);
 
 	// Set up media texture.
@@ -252,6 +253,7 @@ void UMediaPlateComponent::Close()
 
 	StopClockSink();
 	bWantsToPlayWhenVisible = false;
+	bResumeWhenOpened = false;
 	PlaylistIndex = 0;
 }
 
@@ -483,9 +485,13 @@ void UMediaPlateComponent::ResumeWhenVisible()
 		}
 		else if (bWantsToPlayWhenVisible)
 		{
-			Open();
-			FTimespan PlayTime = GetResumeTime();
-			MediaPlayer->Seek(PlayTime);
+			if ((bResumeWhenOpened == false) &&
+				(MediaPlayer->IsPreparing() == false) &&
+				(MediaPlayer->IsPlaying() == false))
+			{
+				bResumeWhenOpened = true;
+				Open();
+			}
 		}
 	}
 }
@@ -499,6 +505,21 @@ FTimespan UMediaPlateComponent::GetResumeTime()
 		float CurrentTime = FApp::GetGameTime();
 		float ElapsedTime = CurrentTime - TimeWhenPlaybackPaused;
 		PlayerTime += FTimespan::FromSeconds(ElapsedTime);
+		
+		// Are we over the length of the media?
+		FTimespan MediaDuration = MediaPlayer->GetDuration();
+		if (PlayerTime > MediaDuration)
+		{
+			bool bIsPlaylist = (MediaPlaylist != nullptr) && (MediaPlaylist->Num() > 1);
+			if ((bLoop) && (bIsPlaylist == false))
+			{
+				PlayerTime %= MediaDuration;
+			}
+			else
+			{
+				PlayerTime = MediaDuration - FTimespan(1);
+			}
+		}
 	}
 
 	return PlayerTime;
@@ -601,6 +622,18 @@ void UMediaPlateComponent::RemoveLetterboxes()
 	Letterboxes.Empty();
 }
 
+void UMediaPlateComponent::OnMediaOpened(FString DeviceUrl)
+{
+	if (bResumeWhenOpened)
+	{
+		bResumeWhenOpened = false;
+		if (MediaPlayer != nullptr)
+		{
+			FTimespan PlayTime = GetResumeTime();
+			MediaPlayer->Seek(PlayTime);
+		}
+	}
+}
 
 void UMediaPlateComponent::OnMediaEnd()
 {
diff --git a/Engine/Plugins/Media/MediaPlate/Source/MediaPlate/Public/MediaPlateComponent.h b/Engine/Plugins/Media/MediaPlate/Source/MediaPlate/Public/MediaPlateComponent.h
--- a/Engine/Plugins/Media/MediaPlate/Source/MediaPlate/Public/MediaPlateComponent.h
+++ b/Engine/Plugins/Media/MediaPlate/Source/MediaPlate/Public/MediaPlateComponent.h
@@ -262,6 +262,8 @@ private:
 	float TimeWhenPlaybackPaused = 0.0f;
 	/** True if our media should be playing when visible. */
 	bool bWantsToPlayWhenVisible = false;
+	/** True if we should resume where we left off when we open the media. */
+	bool bResumeWhenOpened = false;
 
 	/**
 	 * Plays a media source.
@@ -311,6 +313,12 @@ private:
 	 */
 	void RemoveLetterboxes();
 
+	/**
+	 * Called by the media player when the media opens.
+	 */
+	UFUNCTION()
+	void OnMediaOpened(FString DeviceUrl);
+
 	/**
 	 * Called by the media player when the video ends.
 	 */
