MediaPlate: Media details shows number of tiles.
#rnx

diff --git a/Engine/Plugins/Media/MediaPlate/Source/MediaPlateEditor/Private/Widgets/SMediaPlateEditorMediaDetails.cpp b/Engine/Plugins/Media/MediaPlate/Source/MediaPlateEditor/Private/Widgets/SMediaPlateEditorMediaDetails.cpp
--- a/Engine/Plugins/Media/MediaPlate/Source/MediaPlateEditor/Private/Widgets/SMediaPlateEditorMediaDetails.cpp
+++ b/Engine/Plugins/Media/MediaPlate/Source/MediaPlateEditor/Private/Widgets/SMediaPlateEditorMediaDetails.cpp
@@ -60,12 +60,6 @@ void SMediaPlateEditorMediaDetails::Construct(const FArguments& InArgs,
 								[
 									SAssignNew(ResourceSizeText, STextBlock)
 								]
-						]
-
-					// Right side.
-					+ SHorizontalBox::Slot()
-						[
-							SNew(SVerticalBox)
 
 							// Method.
 							+ SVerticalBox::Slot()
@@ -75,6 +69,12 @@ void SMediaPlateEditorMediaDetails::Construct(const FArguments& InArgs,
 								[
 									SAssignNew(MethodText, STextBlock)
 								]
+						]
+
+					// Right side.
+					+ SHorizontalBox::Slot()
+						[
+							SNew(SVerticalBox)
 
 							// Format.
 							+ SVerticalBox::Slot()
@@ -102,6 +102,15 @@ void SMediaPlateEditorMediaDetails::Construct(const FArguments& InArgs,
 								[
 									SAssignNew(NumMipsText, STextBlock)
 								]
+
+							// Num tiles.
+							+ SVerticalBox::Slot()
+								.AutoHeight()
+								.VAlign(VAlign_Center)
+								.Padding(4.0f)
+								[
+									SAssignNew(NumTilesText, STextBlock)
+								]
 						]
 				]
 
@@ -126,6 +135,7 @@ void SMediaPlateEditorMediaDetails::UpdateDetails()
 	int32 LODBias = 0;
 	FText Method;
 	int32 NumMips = 0;
+	int32 NumTotalTiles = 0;
 	int64 ResourceSize = 0;
 	int32 SurfaceWidth = 0;
 	int32 SurfaceHeight = 0;
@@ -138,6 +148,8 @@ void SMediaPlateEditorMediaDetails::UpdateDetails()
 		{
 			FrameRate = MediaPlayer->GetVideoTrackFrameRate(INDEX_NONE, INDEX_NONE);
 			Format = MediaPlayer->GetVideoTrackType(INDEX_NONE, INDEX_NONE);
+			FIntPoint NumTiles = MediaPlayer->GetTileNum();
+			NumTotalTiles = NumTiles.X * NumTiles.Y;
 		}
 
 		// Get texture info,
@@ -166,6 +178,8 @@ void SMediaPlateEditorMediaDetails::UpdateDetails()
 	MethodText->SetText(FText::Format(LOCTEXT("Method", "Method: {0}"), Method));
 	NumMipsText->SetText(FText::Format(LOCTEXT("NumberOfMips", "Number Of Mips: {0}"),
 		FText::AsNumber(NumMips)));
+	NumTilesText->SetText(FText::Format(LOCTEXT("NumberOfTiles", "Number Of Tiles: {0}"),
+		FText::AsNumber(NumTotalTiles)));
 	ResolutionText->SetText(FText::Format(LOCTEXT("Resolution", "Resolution: {0}x{1}"),
 		FText::AsNumber(SurfaceWidth), FText::AsNumber(SurfaceHeight)));
 	ResourceSizeText->SetText(FText::Format(LOCTEXT("ResourceSize", "Resource Size: {0} KB"),
diff --git a/Engine/Plugins/Media/MediaPlate/Source/MediaPlateEditor/Private/Widgets/SMediaPlateEditorMediaDetails.h b/Engine/Plugins/Media/MediaPlate/Source/MediaPlateEditor/Private/Widgets/SMediaPlateEditorMediaDetails.h
--- a/Engine/Plugins/Media/MediaPlate/Source/MediaPlateEditor/Private/Widgets/SMediaPlateEditorMediaDetails.h
+++ b/Engine/Plugins/Media/MediaPlate/Source/MediaPlateEditor/Private/Widgets/SMediaPlateEditorMediaDetails.h
@@ -50,6 +50,7 @@ private:
 	TSharedPtr<STextBlock> LODBiasText;
 	TSharedPtr<STextBlock> MethodText;
 	TSharedPtr<STextBlock> NumMipsText;
+	TSharedPtr<STextBlock> NumTilesText;
 	TSharedPtr<STextBlock> ResolutionText;
 	TSharedPtr<STextBlock> ResourceSizeText;
 };
