MediaPlate: Changed sort order of media plate category to place it at the top of the details pane.
#rnx

diff --git a/Engine/Plugins/Media/MediaPlate/Source/MediaPlateEditor/Private/MediaPlateCustomization.cpp b/Engine/Plugins/Media/MediaPlate/Source/MediaPlateEditor/Private/MediaPlateCustomization.cpp
--- a/Engine/Plugins/Media/MediaPlate/Source/MediaPlateEditor/Private/MediaPlateCustomization.cpp
+++ b/Engine/Plugins/Media/MediaPlate/Source/MediaPlateEditor/Private/MediaPlateCustomization.cpp
@@ -66,6 +66,8 @@ void FMediaPlateCustomization::CustomizeDetails(IDetailLayoutBuilder& DetailBuil
 	StaticMeshCategory.SetCategoryVisibility(false);
 
 	IDetailCategoryBuilder& MediaPlateCategory = DetailBuilder.EditCategory("MediaPlate");
+	// Place this at the top.
+	MediaPlateCategory.SetSortOrder(0);
 
 	// Get objects we are editing.
 	TArray<TWeakObjectPtr<UObject>> Objects;
