MediaPlate: Fixed nDisplay crash when changing material.
#rnx

diff --git a/Engine/Plugins/Media/MediaPlate/Source/MediaPlate/Private/MediaPlate.cpp b/Engine/Plugins/Media/MediaPlate/Source/MediaPlate/Private/MediaPlate.cpp
--- a/Engine/Plugins/Media/MediaPlate/Source/MediaPlate/Private/MediaPlate.cpp
+++ b/Engine/Plugins/Media/MediaPlate/Source/MediaPlate/Private/MediaPlate.cpp
@@ -123,7 +123,7 @@ void AMediaPlate::ApplyCurrentMaterial()
 
 void AMediaPlate::ApplyMaterial(UMaterialInterface* Material)
 {
-	if (Material != nullptr && StaticMeshComponent != nullptr)
+	if (GEditor != nullptr && Material != nullptr && StaticMeshComponent != nullptr)
 	{
 		UMaterialInstanceDynamic* MID = Cast<UMaterialInstanceDynamic>(Material);
 		UMaterialInterface* Result = nullptr;
