MediaPlate: Looping is now enabled by default.
#rnx

diff --git a/Engine/Plugins/Media/MediaPlate/Source/MediaPlate/Public/MediaPlateComponent.h b/Engine/Plugins/Media/MediaPlate/Source/MediaPlate/Public/MediaPlateComponent.h
--- a/Engine/Plugins/Media/MediaPlate/Source/MediaPlate/Public/MediaPlateComponent.h
+++ b/Engine/Plugins/Media/MediaPlate/Source/MediaPlate/Public/MediaPlateComponent.h
@@ -228,7 +228,7 @@ private:
 
 	/** If set then loop when we reach the end. */
 	UPROPERTY(EditAnywhere, Blueprintgetter = GetLoop, BlueprintSetter = SetLoop, Category = "MediaPlate", meta = (AllowPrivateAccess = true))
-	bool bLoop;
+	bool bLoop = true;
 
 	/** Visible mips and tiles calculation mode for the supported mesh types in MediaPlate. (Player restart on change.) */
 	UPROPERTY(EditAnywhere, AdvancedDisplay, Category = "MediaPlate", meta = (AllowPrivateAccess = true))
