MediaPlate: Uses the same aspect ratios as CineCamera.
#rnx

diff --git a/Engine/Plugins/Media/MediaPlate/Source/MediaPlateEditor/MediaPlateEditor.Build.cs b/Engine/Plugins/Media/MediaPlate/Source/MediaPlateEditor/MediaPlateEditor.Build.cs
--- a/Engine/Plugins/Media/MediaPlate/Source/MediaPlateEditor/MediaPlateEditor.Build.cs
+++ b/Engine/Plugins/Media/MediaPlate/Source/MediaPlateEditor/MediaPlateEditor.Build.cs
@@ -18,6 +18,7 @@ public class MediaPlateEditor : ModuleRules
 		PrivateDependencyModuleNames.AddRange(
 			new string[]
 			{
+				"CinematicCamera",
 				"CoreUObject",
 				"DesktopWidgets",
 				"EditorStyle",
diff --git a/Engine/Plugins/Media/MediaPlate/Source/MediaPlateEditor/Private/MediaPlateCustomization.cpp b/Engine/Plugins/Media/MediaPlate/Source/MediaPlateEditor/Private/MediaPlateCustomization.cpp
--- a/Engine/Plugins/Media/MediaPlate/Source/MediaPlateEditor/Private/MediaPlateCustomization.cpp
+++ b/Engine/Plugins/Media/MediaPlate/Source/MediaPlateEditor/Private/MediaPlateCustomization.cpp
@@ -2,6 +2,7 @@
 
 #include "MediaPlateCustomization.h"
 
+#include "CineCameraSettings.h"
 #include "DetailCategoryBuilder.h"
 #include "DetailLayoutBuilder.h"
 #include "DetailWidgetRow.h"
@@ -780,40 +781,34 @@ void FMediaPlateCustomization::SetIsAspectRatioAuto(ECheckBoxState State)
 
 TSharedRef<SWidget> FMediaPlateCustomization::OnGetAspectRatios()
 {
-	FUIAction Actions[4];
-	Actions[0] = FUIAction(FExecuteAction::CreateSP(this, &FMediaPlateCustomization::SetAspectRatio, 16.0f / 9.0f));
-	Actions[1] = FUIAction(FExecuteAction::CreateSP(this, &FMediaPlateCustomization::SetAspectRatio, 16.0f / 10.0f));
-	Actions[2] = FUIAction(FExecuteAction::CreateSP(this, &FMediaPlateCustomization::SetAspectRatio, 4.0f / 3.0f));
-	Actions[3] = FUIAction(FExecuteAction::CreateSP(this, &FMediaPlateCustomization::SetAspectRatio, 1.0f));
-
 	FMenuBuilder MenuBuilder(true, NULL);
-	AddAspectRatiosToMenuBuilder(MenuBuilder, Actions);
+	AddAspectRatiosToMenuBuilder(MenuBuilder, &FMediaPlateCustomization::SetAspectRatio);
 
 	return MenuBuilder.MakeWidget();
 }
 
 TSharedRef<SWidget> FMediaPlateCustomization::OnGetLetterboxAspectRatios()
 {
-	FUIAction Actions[5];
-	Actions[0] = FUIAction(FExecuteAction::CreateSP(this, &FMediaPlateCustomization::SetLetterboxAspectRatio, 16.0f / 9.0f));
-	Actions[1] = FUIAction(FExecuteAction::CreateSP(this, &FMediaPlateCustomization::SetLetterboxAspectRatio, 16.0f / 10.0f));
-	Actions[2] = FUIAction(FExecuteAction::CreateSP(this, &FMediaPlateCustomization::SetLetterboxAspectRatio, 4.0f / 3.0f));
-	Actions[3] = FUIAction(FExecuteAction::CreateSP(this, &FMediaPlateCustomization::SetLetterboxAspectRatio, 1.0f));
-	Actions[4] = FUIAction(FExecuteAction::CreateSP(this, &FMediaPlateCustomization::SetLetterboxAspectRatio, 0.0f));
-
 	FMenuBuilder MenuBuilder(true, NULL);
-	AddAspectRatiosToMenuBuilder(MenuBuilder, Actions);
-	MenuBuilder.AddMenuEntry(LOCTEXT("Disable", "Disable"), FText(), FSlateIcon(), Actions[4]);
+	AddAspectRatiosToMenuBuilder(MenuBuilder, &FMediaPlateCustomization::SetLetterboxAspectRatio);
+	
+	FUIAction Action = FUIAction(FExecuteAction::CreateSP(this, &FMediaPlateCustomization::SetLetterboxAspectRatio, 0.0f));
+	MenuBuilder.AddMenuEntry(LOCTEXT("Disable", "Disable"), FText(), FSlateIcon(), Action);
 
 	return MenuBuilder.MakeWidget();
 }
 
-void FMediaPlateCustomization::AddAspectRatiosToMenuBuilder(FMenuBuilder& MenuBuilder, FUIAction Actions[])
+void FMediaPlateCustomization::AddAspectRatiosToMenuBuilder(FMenuBuilder& MenuBuilder,
+	void (FMediaPlateCustomization::* Func)(float))
 {
-	MenuBuilder.AddMenuEntry(LOCTEXT("16x9", "16x9"), FText(), FSlateIcon(), Actions[0]);
-	MenuBuilder.AddMenuEntry(LOCTEXT("16x10", "16x10"), FText(), FSlateIcon(), Actions[1]);
-	MenuBuilder.AddMenuEntry(LOCTEXT("4x3", "4x3"), FText(), FSlateIcon(), Actions[2]);
-	MenuBuilder.AddMenuEntry(LOCTEXT("1x1", "1x1"), FText(), FSlateIcon(), Actions[3]);
+	TArray<FNamedFilmbackPreset> const& Presets = UCineCameraSettings::GetFilmbackPresets();
+
+	for (const FNamedFilmbackPreset& Preset : Presets)
+	{
+		FUIAction Action = FUIAction(FExecuteAction::CreateSP(this,
+			Func, Preset.FilmbackSettings.SensorAspectRatio));
+		MenuBuilder.AddMenuEntry(FText::FromString(Preset.Name), FText(), FSlateIcon(), Action);
+	}
 }
 
 void FMediaPlateCustomization::SetAspectRatio(float AspectRatio)
diff --git a/Engine/Plugins/Media/MediaPlate/Source/MediaPlateEditor/Private/MediaPlateCustomization.h b/Engine/Plugins/Media/MediaPlate/Source/MediaPlateEditor/Private/MediaPlateCustomization.h
--- a/Engine/Plugins/Media/MediaPlate/Source/MediaPlateEditor/Private/MediaPlateCustomization.h
+++ b/Engine/Plugins/Media/MediaPlate/Source/MediaPlateEditor/Private/MediaPlateCustomization.h
@@ -100,9 +100,10 @@ private:
 	TSharedRef<SWidget> OnGetLetterboxAspectRatios();
 
 	/**
-	 * Returns menu options for all aspect ratio presets.
+	 * Adds menu options for all aspect ratio presets.
 	 */
-	void AddAspectRatiosToMenuBuilder(FMenuBuilder& MenuBuilder, FUIAction Actions[]);
+	void AddAspectRatiosToMenuBuilder(FMenuBuilder& MenuBuilder,
+		void (FMediaPlateCustomization::*Func)(float));
 
 	/**
 	 * Call this to see if auto aspect ratio is enabled.
