MediaPlate: Prompts the user to disable autoplay when binding to Sequencer.
#rnx

diff --git a/Engine/Plugins/Media/MediaPlate/Source/MediaPlateEditor/Private/Sequencer/MediaPlateTrackEditor.cpp b/Engine/Plugins/Media/MediaPlate/Source/MediaPlateEditor/Private/Sequencer/MediaPlateTrackEditor.cpp
--- a/Engine/Plugins/Media/MediaPlate/Source/MediaPlateEditor/Private/Sequencer/MediaPlateTrackEditor.cpp
+++ b/Engine/Plugins/Media/MediaPlate/Source/MediaPlateEditor/Private/Sequencer/MediaPlateTrackEditor.cpp
@@ -10,6 +10,7 @@
 #include "MediaPlateComponent.h"
 #include "MediaPlayer.h"
 #include "MediaPlaylist.h"
+#include "Misc/MessageDialog.h"
 #include "Modules/ModuleManager.h"
 #include "MovieSceneMediaSection.h"
 #include "MovieSceneMediaTrack.h"
@@ -204,6 +205,22 @@ void FMediaPlateTrackEditor::AddTrackForComponent(UMediaPlateComponent* Componen
 			}
 		}
 	}
+
+	// Does the media plate have autoplay?
+	if (Component->bAutoPlay)
+	{
+		// Ask the user if it can be disabled.
+		EAppReturnType::Type YesNoCancelReply = FMessageDialog::Open(EAppMsgType::YesNo,
+			LOCTEXT("Prompt_DisableAutoplay", "Would you like to disable autoplay on the MediaPlate actor (recommended)?"));
+
+		switch (YesNoCancelReply)
+		{
+			case EAppReturnType::Yes:
+				Component->Modify();
+				Component->bAutoPlay = false;
+				break;
+		}
+	}
 }
 
 void FMediaPlateTrackEditor::OnRelease()
