MediaPlate: Autoplay enabled by default.
#rnx

diff --git a/Engine/Plugins/Media/MediaPlate/Source/MediaPlate/Public/MediaPlateComponent.h b/Engine/Plugins/Media/MediaPlate/Source/MediaPlate/Public/MediaPlateComponent.h
--- a/Engine/Plugins/Media/MediaPlate/Source/MediaPlate/Public/MediaPlateComponent.h
+++ b/Engine/Plugins/Media/MediaPlate/Source/MediaPlate/Public/MediaPlateComponent.h
@@ -115,7 +115,7 @@ public:
 
 	/** If set then start playing when this object is active. */
 	UPROPERTY(EditAnywhere, Category = "MediaPlate")
-	bool bAutoPlay;
+	bool bAutoPlay = true;
 
 	/** If set then enable audio. */
 	UPROPERTY(EditAnywhere, Category = "MediaPlate")
