MediaPlate: Fixed properties not being marked instance which caused errors when saving out.
#rnx

diff --git a/Engine/Plugins/Media/MediaPlate/Source/MediaPlate/Public/MediaPlateComponent.h b/Engine/Plugins/Media/MediaPlate/Source/MediaPlate/Public/MediaPlateComponent.h
--- a/Engine/Plugins/Media/MediaPlate/Source/MediaPlate/Public/MediaPlateComponent.h
+++ b/Engine/Plugins/Media/MediaPlate/Source/MediaPlate/Public/MediaPlateComponent.h
@@ -141,7 +141,7 @@ public:
 	TArray<TObjectPtr<UStaticMeshComponent>> Letterboxes;
 
 	/** What media playlist to play. */
-	UPROPERTY(BlueprintReadWrite, Category = "MediaPlate")
+	UPROPERTY(Instanced, BlueprintReadWrite, Category = "MediaPlate")
 	TObjectPtr<UMediaPlaylist> MediaPlaylist;
 
 	/** The current index of the source in the play list being played. */
@@ -252,11 +252,11 @@ private:
 	static FLazyName MediaPlaylistName;
 
 	/** Holds the media player. */
-	UPROPERTY()
+	UPROPERTY(Instanced)
 	TObjectPtr<UMediaTexture> MediaTexture;
 
 	/** This component's media player */
-	UPROPERTY()
+	UPROPERTY(Instanced)
 	TObjectPtr<UMediaPlayer> MediaPlayer;
 
 	/** Info representing this object. */
