MediaPlate: Fixed close icon using wrong icon.
#rnx

diff --git a/Engine/Plugins/Media/MediaPlate/Source/MediaPlateEditor/Private/MediaPlateCustomization.cpp b/Engine/Plugins/Media/MediaPlate/Source/MediaPlateEditor/Private/MediaPlateCustomization.cpp
--- a/Engine/Plugins/Media/MediaPlate/Source/MediaPlateEditor/Private/MediaPlateCustomization.cpp
+++ b/Engine/Plugins/Media/MediaPlate/Source/MediaPlateEditor/Private/MediaPlateCustomization.cpp
@@ -437,7 +437,7 @@ void FMediaPlateCustomization::CustomizeDetails(IDetailLayoutBuilder& DetailBuil
 						[
 							SNew(SImage)
 								.ColorAndOpacity(FSlateColor::UseForeground())
-								.Image(Style->GetBrush("MediaPlateEditor.StopMedia.Small"))
+								.Image(Style->GetBrush("MediaPlateEditor.CloseMedia.Small"))
 						]
 				]
 		];
diff --git a/Engine/Plugins/Media/MediaPlate/Source/MediaPlateEditor/Private/MediaPlateEditorStyle.cpp b/Engine/Plugins/Media/MediaPlate/Source/MediaPlateEditor/Private/MediaPlateEditorStyle.cpp
--- a/Engine/Plugins/Media/MediaPlate/Source/MediaPlateEditor/Private/MediaPlateEditorStyle.cpp
+++ b/Engine/Plugins/Media/MediaPlate/Source/MediaPlateEditor/Private/MediaPlateEditorStyle.cpp
@@ -50,6 +50,7 @@ FMediaPlateEditorStyle::FMediaPlateEditorStyle()
 
 	// toolbar icons
 	Set("MediaPlateEditor.CloseMedia", new IMAGE_BRUSH("icon_eject_40x", Icon40x40));
+	Set("MediaPlateEditor.CloseMedia.Small", new IMAGE_BRUSH("icon_eject_40x", Icon20x20));
 	Set("MediaPlateEditor.ForwardMedia", new IMAGE_BRUSH("icon_forward_40x", Icon40x40));
 	Set("MediaPlateEditor.ForwardMedia.Small", new IMAGE_BRUSH("icon_forward_40x", Icon20x20));
 	Set("MediaPlateEditor.NextMedia", new IMAGE_BRUSH("icon_step_40x", Icon40x40));
