MediaPlate: Fixed open icon using the wrong icon.
#rnx

diff --git a/Engine/Plugins/Media/MediaPlate/Source/MediaPlateEditor/Private/MediaPlateEditorStyle.cpp b/Engine/Plugins/Media/MediaPlate/Source/MediaPlateEditor/Private/MediaPlateEditorStyle.cpp
--- a/Engine/Plugins/Media/MediaPlate/Source/MediaPlateEditor/Private/MediaPlateEditorStyle.cpp
+++ b/Engine/Plugins/Media/MediaPlate/Source/MediaPlateEditor/Private/MediaPlateEditorStyle.cpp
@@ -55,8 +55,8 @@ FMediaPlateEditorStyle::FMediaPlateEditorStyle()
 	Set("MediaPlateEditor.ForwardMedia.Small", new IMAGE_BRUSH("icon_forward_40x", Icon20x20));
 	Set("MediaPlateEditor.NextMedia", new IMAGE_BRUSH("icon_step_40x", Icon40x40));
 	Set("MediaPlateEditor.NextMedia.Small", new IMAGE_BRUSH("icon_step_40x", Icon20x20));
-	Set("MediaPlateEditor.OpenMedia", new IMAGE_BRUSH("tab_media_16x", Icon40x40));
-	Set("MediaPlateEditor.OpenMedia.Small", new IMAGE_BRUSH("tab_media_16x", Icon20x20));
+	Set("MediaPlateEditor.OpenMedia", new IMAGE_BRUSH("icon_open_40x", Icon40x40));
+	Set("MediaPlateEditor.OpenMedia.Small", new IMAGE_BRUSH("icon_open_40x", Icon20x20));
 	Set("MediaPlateEditor.PauseMedia", new IMAGE_BRUSH("icon_pause_40x", Icon40x40));
 	Set("MediaPlateEditor.PauseMedia.Small", new IMAGE_BRUSH("icon_pause_40x", Icon20x20));
 	Set("MediaPlateEditor.PlayMedia", new IMAGE_BRUSH("icon_play_40x", Icon40x40));
