MediaPlate: Added tooltips to control buttons.
#rnx

diff --git a/Engine/Plugins/Media/MediaPlate/Source/MediaPlateEditor/Private/MediaPlateCustomization.cpp b/Engine/Plugins/Media/MediaPlate/Source/MediaPlateEditor/Private/MediaPlateCustomization.cpp
--- a/Engine/Plugins/Media/MediaPlate/Source/MediaPlateEditor/Private/MediaPlateCustomization.cpp
+++ b/Engine/Plugins/Media/MediaPlate/Source/MediaPlateEditor/Private/MediaPlateCustomization.cpp
@@ -218,8 +218,9 @@ void FMediaPlateCustomization::CustomizeDetails(IDetailLayoutBuilder& DetailBuil
 					})
 					[
 						SNew(SImage)
-						.ColorAndOpacity(FSlateColor::UseForeground())
-					.Image(Style->GetBrush("MediaPlateEditor.RewindMedia.Small"))
+							.ColorAndOpacity(FSlateColor::UseForeground())
+							.Image(Style->GetBrush("MediaPlateEditor.RewindMedia.Small"))
+							.ToolTipText(LOCTEXT("Rewind", "Rewind the media to the beginning"))
 					]
 				]
 
@@ -265,6 +266,7 @@ void FMediaPlateCustomization::CustomizeDetails(IDetailLayoutBuilder& DetailBuil
 							SNew(SImage)
 								.ColorAndOpacity(FSlateColor::UseForeground())
 								.Image(Style->GetBrush("MediaPlateEditor.ReverseMedia.Small"))
+								.ToolTipText(LOCTEXT("Reverse", "Reverse media playback"))
 						]
 				]
 
@@ -307,6 +309,7 @@ void FMediaPlateCustomization::CustomizeDetails(IDetailLayoutBuilder& DetailBuil
 							SNew(SImage)
 								.ColorAndOpacity(FSlateColor::UseForeground())
 								.Image(Style->GetBrush("MediaPlateEditor.PlayMedia.Small"))
+								.ToolTipText(LOCTEXT("Play", "Start media playback"))
 						]
 				]
 
@@ -333,6 +336,7 @@ void FMediaPlateCustomization::CustomizeDetails(IDetailLayoutBuilder& DetailBuil
 						SNew(SImage)
 							.ColorAndOpacity(FSlateColor::UseForeground())
 							.Image(Style->GetBrush("MediaPlateEditor.PauseMedia.Small"))
+							.ToolTipText(LOCTEXT("Pause", "Pause media playback"))
 					]
 				]
 
@@ -378,6 +382,7 @@ void FMediaPlateCustomization::CustomizeDetails(IDetailLayoutBuilder& DetailBuil
 						SNew(SImage)
 							.ColorAndOpacity(FSlateColor::UseForeground())
 							.Image(Style->GetBrush("MediaPlateEditor.ForwardMedia.Small"))
+							.ToolTipText(LOCTEXT("Forward", "Fast forward media playback"))
 					]
 				]
 
@@ -412,6 +417,7 @@ void FMediaPlateCustomization::CustomizeDetails(IDetailLayoutBuilder& DetailBuil
 							SNew(SImage)
 								.ColorAndOpacity(FSlateColor::UseForeground())
 								.Image(Style->GetBrush("MediaPlateEditor.OpenMedia.Small"))
+								.ToolTipText(LOCTEXT("Open", "Open the current media"))
 						]
 				]
 
@@ -438,6 +444,7 @@ void FMediaPlateCustomization::CustomizeDetails(IDetailLayoutBuilder& DetailBuil
 							SNew(SImage)
 								.ColorAndOpacity(FSlateColor::UseForeground())
 								.Image(Style->GetBrush("MediaPlateEditor.CloseMedia.Small"))
+								.ToolTipText(LOCTEXT("Close", "Close the currently opened media"))
 						]
 				]
 		];
