MediaPlate: Control buttons are now only enabled when usable, similar to the media player.
#rnx

diff --git a/Engine/Plugins/Media/MediaPlate/Source/MediaPlateEditor/Private/MediaPlateCustomization.cpp b/Engine/Plugins/Media/MediaPlate/Source/MediaPlateEditor/Private/MediaPlateCustomization.cpp
--- a/Engine/Plugins/Media/MediaPlate/Source/MediaPlateEditor/Private/MediaPlateCustomization.cpp
+++ b/Engine/Plugins/Media/MediaPlate/Source/MediaPlateEditor/Private/MediaPlateCustomization.cpp
@@ -199,29 +199,47 @@ void FMediaPlateCustomization::CustomizeDetails(IDetailLayoutBuilder& DetailBuil
 				.AutoWidth()
 				[
 					SNew(SButton)
-					.VAlign(VAlign_Center)
-					.OnClicked_Lambda([this]() -> FReply
-					{
-						for (TWeakObjectPtr<UMediaPlateComponent>& MediaPlatePtr : MediaPlatesList)
+						.VAlign(VAlign_Center)
+						.IsEnabled_Lambda([this]
 						{
-							UMediaPlateComponent* MediaPlate = MediaPlatePtr.Get();
-							if (MediaPlate != nullptr)
+							for (TWeakObjectPtr<UMediaPlateComponent>& MediaPlatePtr : MediaPlatesList)
 							{
-								UMediaPlayer* MediaPlayer = MediaPlate->GetMediaPlayer();
-								if (MediaPlayer != nullptr)
+								UMediaPlateComponent* MediaPlate = MediaPlatePtr.Get();
+								if (MediaPlate != nullptr)
 								{
-									MediaPlayer->Rewind();
+									UMediaPlayer* MediaPlayer = MediaPlate->GetMediaPlayer();
+									if (MediaPlayer != nullptr)
+									{
+										return MediaPlayer->IsReady() &&
+											MediaPlayer->SupportsSeeking() &&
+											MediaPlayer->GetTime() > FTimespan::Zero();
+									}
 								}
 							}
-						}
-						return FReply::Handled();
-					})
-					[
-						SNew(SImage)
-							.ColorAndOpacity(FSlateColor::UseForeground())
-							.Image(Style->GetBrush("MediaPlateEditor.RewindMedia.Small"))
-							.ToolTipText(LOCTEXT("Rewind", "Rewind the media to the beginning"))
-					]
+							return false;
+						})
+						.OnClicked_Lambda([this]() -> FReply
+						{
+							for (TWeakObjectPtr<UMediaPlateComponent>& MediaPlatePtr : MediaPlatesList)
+							{
+								UMediaPlateComponent* MediaPlate = MediaPlatePtr.Get();
+								if (MediaPlate != nullptr)
+								{
+									UMediaPlayer* MediaPlayer = MediaPlate->GetMediaPlayer();
+									if (MediaPlayer != nullptr)
+									{
+										MediaPlayer->Rewind();
+									}
+								}
+							}
+							return FReply::Handled();
+						})
+						[
+							SNew(SImage)
+								.ColorAndOpacity(FSlateColor::UseForeground())
+								.Image(Style->GetBrush("MediaPlateEditor.RewindMedia.Small"))
+								.ToolTipText(LOCTEXT("Rewind", "Rewind the media to the beginning"))
+						]
 				]
 
 			// Reverse button.
@@ -319,25 +337,41 @@ void FMediaPlateCustomization::CustomizeDetails(IDetailLayoutBuilder& DetailBuil
 				.AutoWidth()
 				[
 					SNew(SButton)
-					.VAlign(VAlign_Center)
-					.OnClicked_Lambda([this]() -> FReply
-					{
-						for (TWeakObjectPtr<UMediaPlateComponent>& MediaPlatePtr : MediaPlatesList)
+						.VAlign(VAlign_Center)
+						.IsEnabled_Lambda([this]
 						{
-							UMediaPlateComponent* MediaPlate = MediaPlatePtr.Get();
-							if (MediaPlate != nullptr)
-							{							
-								MediaPlate->Pause();
+							for (TWeakObjectPtr<UMediaPlateComponent>& MediaPlatePtr : MediaPlatesList)
+							{
+								UMediaPlateComponent* MediaPlate = MediaPlatePtr.Get();
+								if (MediaPlate != nullptr)
+								{
+									UMediaPlayer* MediaPlayer = MediaPlate->GetMediaPlayer();
+									if (MediaPlayer != nullptr)
+									{
+										return (MediaPlayer->CanPause()) && (!MediaPlayer->IsPaused());
+									}
+								}
 							}
-						}
-						return FReply::Handled();
-					})
-					[
-						SNew(SImage)
-							.ColorAndOpacity(FSlateColor::UseForeground())
-							.Image(Style->GetBrush("MediaPlateEditor.PauseMedia.Small"))
-							.ToolTipText(LOCTEXT("Pause", "Pause media playback"))
-					]
+							return false;
+						})
+						.OnClicked_Lambda([this]() -> FReply
+						{
+							for (TWeakObjectPtr<UMediaPlateComponent>& MediaPlatePtr : MediaPlatesList)
+							{
+								UMediaPlateComponent* MediaPlate = MediaPlatePtr.Get();
+								if (MediaPlate != nullptr)
+								{							
+									MediaPlate->Pause();
+								}
+							}
+							return FReply::Handled();
+						})
+						[
+							SNew(SImage)
+								.ColorAndOpacity(FSlateColor::UseForeground())
+								.Image(Style->GetBrush("MediaPlateEditor.PauseMedia.Small"))
+								.ToolTipText(LOCTEXT("Pause", "Pause media playback"))
+						]
 				]
 
 			// Forward button.
@@ -428,6 +462,22 @@ void FMediaPlateCustomization::CustomizeDetails(IDetailLayoutBuilder& DetailBuil
 				[
 					SNew(SButton)
 						.VAlign(VAlign_Center)
+						.IsEnabled_Lambda([this]
+						{
+							for (TWeakObjectPtr<UMediaPlateComponent>& MediaPlatePtr : MediaPlatesList)
+							{
+								UMediaPlateComponent* MediaPlate = MediaPlatePtr.Get();
+								if (MediaPlate != nullptr)
+								{
+									UMediaPlayer* MediaPlayer = MediaPlate->GetMediaPlayer();
+									if (MediaPlayer != nullptr)
+									{
+										return !MediaPlayer->GetUrl().IsEmpty();
+									}
+								}
+							}
+							return false;
+						})
 						.OnClicked_Lambda([this]() -> FReply
 						{
 							for (TWeakObjectPtr<UMediaPlateComponent>& MediaPlatePtr : MediaPlatesList)
