MediaPlate: Checks if the object is hidden before starting playback.
#rnx

diff --git a/Engine/Plugins/Media/MediaPlate/Source/MediaPlate/Private/MediaPlateComponent.cpp b/Engine/Plugins/Media/MediaPlate/Source/MediaPlate/Private/MediaPlateComponent.cpp
--- a/Engine/Plugins/Media/MediaPlate/Source/MediaPlate/Private/MediaPlateComponent.cpp
+++ b/Engine/Plugins/Media/MediaPlate/Source/MediaPlate/Private/MediaPlateComponent.cpp
@@ -188,7 +188,11 @@ void UMediaPlateComponent::Open()
 {
 	bIsMediaPlatePlaying = true;
 	CurrentRate = bPlayOnOpen ? 1.0f : 0.0f;
-	if ((bPlayOnlyWhenVisible == false) || (IsVisible()))
+
+	// Is play only when visible off, and we should be visible?
+	// Or is play only when visible on, and we are visible?
+	if (((bPlayOnlyWhenVisible == false) && (ShouldBeVisible())) ||
+		(IsVisible()))
 	{
 		bool bIsPlaying = false;
 		if (MediaPlayer != nullptr)
@@ -501,6 +505,11 @@ bool UMediaPlateComponent::IsVisible()
 	return GetOwner()->WasRecentlyRendered();
 }
 
+bool UMediaPlateComponent::ShouldBeVisible()
+{
+	return ((StaticMeshComponent != nullptr) && (StaticMeshComponent->ShouldRender()));
+}
+
 void UMediaPlateComponent::ResumeWhenVisible()
 {
 	if (MediaPlayer != nullptr)
diff --git a/Engine/Plugins/Media/MediaPlate/Source/MediaPlate/Public/MediaPlateComponent.h b/Engine/Plugins/Media/MediaPlate/Source/MediaPlate/Public/MediaPlateComponent.h
--- a/Engine/Plugins/Media/MediaPlate/Source/MediaPlate/Public/MediaPlateComponent.h
+++ b/Engine/Plugins/Media/MediaPlate/Source/MediaPlate/Public/MediaPlateComponent.h
@@ -307,6 +307,14 @@ private:
 	 */
 	bool IsVisible();
 
+	/**
+	 * Call this to see if this media plate should be visible
+	 * (but may not necessarily BE visible).
+	 * For example, this will return false if the object is marked as hidden.
+	 * Otherwise this will reeturn true, but IsVisible could still be false if this is off screen.
+	 */
+	bool ShouldBeVisible();
+
 	/**
 	 * Call this to resume playback when the media plate is visible.
 	 */
