MediaPlate: Fixed playback not resuming when disabling bPlayOnlyWhenVisible when plate is off screen.
#rnx

diff --git a/Engine/Plugins/Media/MediaPlate/Source/MediaPlate/Private/MediaPlateComponent.cpp b/Engine/Plugins/Media/MediaPlate/Source/MediaPlate/Private/MediaPlateComponent.cpp
--- a/Engine/Plugins/Media/MediaPlate/Source/MediaPlate/Private/MediaPlateComponent.cpp
+++ b/Engine/Plugins/Media/MediaPlate/Source/MediaPlate/Private/MediaPlateComponent.cpp
@@ -293,14 +293,19 @@ void UMediaPlateComponent::SetMeshRange(FVector2D InMeshRange)
 }
 
 void UMediaPlateComponent::SetPlayOnlyWhenVisible(bool bInPlayOnlyWhenVisible)
+{
+	bPlayOnlyWhenVisible = bInPlayOnlyWhenVisible;
+	PlayOnlyWhenVisibleChanged();
+}
+
+void UMediaPlateComponent::PlayOnlyWhenVisibleChanged()
 {
 	// If we are turning off PlayOnlyWhenVisible then make sure we are playing.
-	if (bInPlayOnlyWhenVisible == false)
+	if (bPlayOnlyWhenVisible == false)
 	{
 		ResumeWhenVisible();
 	}
 
-	bPlayOnlyWhenVisible = bInPlayOnlyWhenVisible;
 	UpdateTicking();
 }
 
@@ -745,6 +750,10 @@ void UMediaPlateComponent::PostEditChangeProperty(FPropertyChangedEvent& Propert
 			}
 		}
 	}
+	else if (PropertyChangedEvent.GetPropertyName() == GET_MEMBER_NAME_CHECKED(ThisClass, bPlayOnlyWhenVisible))
+	{
+		PlayOnlyWhenVisibleChanged();
+	}
 	else if (PropertyChangedEvent.GetPropertyName() == GET_MEMBER_NAME_CHECKED(ThisClass, VisibleMipsTilesCalculations))
 	{
 		if (MediaTextureTrackerObject != nullptr)
diff --git a/Engine/Plugins/Media/MediaPlate/Source/MediaPlate/Public/MediaPlateComponent.h b/Engine/Plugins/Media/MediaPlate/Source/MediaPlate/Public/MediaPlateComponent.h
--- a/Engine/Plugins/Media/MediaPlate/Source/MediaPlate/Public/MediaPlateComponent.h
+++ b/Engine/Plugins/Media/MediaPlate/Source/MediaPlate/Public/MediaPlateComponent.h
@@ -228,6 +228,11 @@ public:
 	virtual bool IsExternalControlAllowed() override;
 
 private:
+	/**
+	 * Should be called when bPlayOnlyWhenVisible changes.
+	 */
+	void PlayOnlyWhenVisibleChanged();
+
 	void RestartPlayer();
 
 	/**
