MediaPlate: Fixed playback trying to resume (after becoming visible) when it should be paused.
#rnx

diff --git a/Engine/Plugins/Media/MediaPlate/Source/MediaPlate/Private/MediaPlateComponent.cpp b/Engine/Plugins/Media/MediaPlate/Source/MediaPlate/Private/MediaPlateComponent.cpp
--- a/Engine/Plugins/Media/MediaPlate/Source/MediaPlate/Private/MediaPlateComponent.cpp
+++ b/Engine/Plugins/Media/MediaPlate/Source/MediaPlate/Private/MediaPlateComponent.cpp
@@ -521,9 +521,13 @@ void UMediaPlateComponent::ResumeWhenVisible()
 	{
 		if (MediaPlayer->IsPaused())
 		{
-			FTimespan PlayTime = GetResumeTime();
-			MediaPlayer->Seek(PlayTime);
-			MediaPlayer->Play();
+			// Should we be playing?
+			if (CurrentRate != 0.0f)
+			{
+				FTimespan PlayTime = GetResumeTime();
+				MediaPlayer->Seek(PlayTime);
+				MediaPlayer->Play();
+			}
 		}
 		else if (bWantsToPlayWhenVisible)
 		{
