MediaPlate: Checking for ticking is now done when the media opens, instead of OnRegister, so it can work with Sequencer.
#rnx

diff --git a/Engine/Plugins/Media/MediaPlate/Source/MediaPlate/Private/MediaPlateComponent.cpp b/Engine/Plugins/Media/MediaPlate/Source/MediaPlate/Private/MediaPlateComponent.cpp
--- a/Engine/Plugins/Media/MediaPlate/Source/MediaPlate/Private/MediaPlateComponent.cpp
+++ b/Engine/Plugins/Media/MediaPlate/Source/MediaPlate/Private/MediaPlateComponent.cpp
@@ -113,9 +113,6 @@ void UMediaPlateComponent::OnRegister()
 			SoundComponent->SetMediaPlayer(MediaPlayer);
 		}
 	}
-
-	// Activate tick if needed.
-	UpdateTicking();
 }
 
 void UMediaPlateComponent::BeginPlay()
@@ -575,7 +572,7 @@ FTimespan UMediaPlateComponent::GetResumeTime()
 
 void UMediaPlateComponent::UpdateTicking()
 {
-	bool bEnableTick = bPlayOnlyWhenVisible;
+	bool bEnableTick = bIsMediaPlatePlaying && bPlayOnlyWhenVisible;
 	PrimaryComponentTick.SetTickFunctionEnable(bEnableTick);
 }
 
@@ -682,6 +679,9 @@ void UMediaPlateComponent::OnMediaOpened(FString DeviceUrl)
 			MediaPlayer->Seek(PlayTime);
 		}
 	}
+
+	// Activate tick if needed.
+	UpdateTicking();
 }
 
 void UMediaPlateComponent::OnMediaEnd()
