MediaPlate: Player will stop playing when object is hidden during playback.
#rnx

diff --git a/Engine/Plugins/Media/MediaPlate/Source/MediaPlate/Private/MediaPlateComponent.cpp b/Engine/Plugins/Media/MediaPlate/Source/MediaPlate/Private/MediaPlateComponent.cpp
--- a/Engine/Plugins/Media/MediaPlate/Source/MediaPlate/Private/MediaPlateComponent.cpp
+++ b/Engine/Plugins/Media/MediaPlate/Source/MediaPlate/Private/MediaPlateComponent.cpp
@@ -150,7 +150,7 @@ void UMediaPlateComponent::TickComponent(float DeltaTime, enum ELevelTick TickTy
 {
 	Super::TickComponent(DeltaTime, TickType, ThisTickFunction);
 
-	if ((MediaPlayer != nullptr) && (bPlayOnlyWhenVisible))
+	if (MediaPlayer != nullptr)
 	{
 		if ((CurrentRate != 0.0f) || (bWantsToPlayWhenVisible))
 		{
@@ -186,10 +186,7 @@ void UMediaPlateComponent::Open()
 	bIsMediaPlatePlaying = true;
 	CurrentRate = bPlayOnOpen ? 1.0f : 0.0f;
 
-	// Is play only when visible off, and we should be visible?
-	// Or is play only when visible on, and we are visible?
-	if (((bPlayOnlyWhenVisible == false) && (ShouldBeVisible())) ||
-		(IsVisible()))
+	if (IsVisible())
 	{
 		bool bIsPlaying = false;
 		if (MediaPlayer != nullptr)
@@ -213,6 +210,8 @@ void UMediaPlateComponent::Open()
 		bWantsToPlayWhenVisible = true;
 		TimeWhenPlaybackPaused = FApp::GetGameTime();
 	}
+
+	UpdateTicking();
 }
 
 void UMediaPlateComponent::Play()
@@ -260,6 +259,7 @@ void UMediaPlateComponent::Close()
 	bWantsToPlayWhenVisible = false;
 	bResumeWhenOpened = false;
 	PlaylistIndex = 0;
+	UpdateTicking();
 }
 
 bool UMediaPlateComponent::GetLoop()
@@ -302,8 +302,6 @@ void UMediaPlateComponent::PlayOnlyWhenVisibleChanged()
 	{
 		ResumeWhenVisible();
 	}
-
-	UpdateTicking();
 }
 
 void UMediaPlateComponent::RegisterWithMediaTextureTracker()
@@ -474,8 +472,8 @@ bool UMediaPlateComponent::SetProxyRate(float Rate)
 
 bool UMediaPlateComponent::IsExternalControlAllowed()
 {
-	// Allow control if we always play, or are visible.
-	return (bPlayOnlyWhenVisible == false) || IsVisible();
+	// Allow control if we are visible.
+	return IsVisible();
 }
 
 void UMediaPlateComponent::RestartPlayer()
@@ -504,12 +502,14 @@ void UMediaPlateComponent::StopClockSink()
 
 bool UMediaPlateComponent::IsVisible()
 {
-	return GetOwner()->WasRecentlyRendered();
-}
+	bool bIsVisible = ((StaticMeshComponent != nullptr) && (StaticMeshComponent->ShouldRender()));
 
-bool UMediaPlateComponent::ShouldBeVisible()
-{
-	return ((StaticMeshComponent != nullptr) && (StaticMeshComponent->ShouldRender()));
+	if (bIsVisible && bPlayOnlyWhenVisible)
+	{
+		bIsVisible = GetOwner()->WasRecentlyRendered();
+	}
+
+	return bIsVisible;
 }
 
 void UMediaPlateComponent::ResumeWhenVisible()
@@ -572,7 +572,7 @@ FTimespan UMediaPlateComponent::GetResumeTime()
 
 void UMediaPlateComponent::UpdateTicking()
 {
-	bool bEnableTick = bIsMediaPlatePlaying && bPlayOnlyWhenVisible;
+	bool bEnableTick = bIsMediaPlatePlaying;
 	PrimaryComponentTick.SetTickFunctionEnable(bEnableTick);
 }
 
@@ -679,9 +679,6 @@ void UMediaPlateComponent::OnMediaOpened(FString DeviceUrl)
 			MediaPlayer->Seek(PlayTime);
 		}
 	}
-
-	// Activate tick if needed.
-	UpdateTicking();
 }
 
 void UMediaPlateComponent::OnMediaEnd()
diff --git a/Engine/Plugins/Media/MediaPlate/Source/MediaPlate/Public/MediaPlateComponent.h b/Engine/Plugins/Media/MediaPlate/Source/MediaPlate/Public/MediaPlateComponent.h
--- a/Engine/Plugins/Media/MediaPlate/Source/MediaPlate/Public/MediaPlateComponent.h
+++ b/Engine/Plugins/Media/MediaPlate/Source/MediaPlate/Public/MediaPlateComponent.h
@@ -312,14 +312,6 @@ private:
 	 */
 	bool IsVisible();
 
-	/**
-	 * Call this to see if this media plate should be visible
-	 * (but may not necessarily BE visible).
-	 * For example, this will return false if the object is marked as hidden.
-	 * Otherwise this will reeturn true, but IsVisible could still be false if this is off screen.
-	 */
-	bool ShouldBeVisible();
-
 	/**
 	 * Call this to resume playback when the media plate is visible.
 	 */
