MediaPlate: Fix mip bias propagation to the  material's media texture sampler.

#rnx

diff --git a/Engine/Plugins/Media/MediaPlate/Source/MediaPlate/Private/MediaPlateComponent.cpp b/Engine/Plugins/Media/MediaPlate/Source/MediaPlate/Private/MediaPlateComponent.cpp
--- a/Engine/Plugins/Media/MediaPlate/Source/MediaPlate/Private/MediaPlateComponent.cpp
+++ b/Engine/Plugins/Media/MediaPlate/Source/MediaPlate/Private/MediaPlateComponent.cpp
@@ -105,6 +105,14 @@ void UMediaPlateComponent::OnRegister()
 		MediaTexture->SetMipMapBias(MipMapBias);
 		MediaTexture->SetMediaPlayer(MediaPlayer);
 		MediaTexture->UpdateResource();
+
+		if (AMediaPlate* MediaPlate = GetOwner<AMediaPlate>())
+		{
+			if (UMaterialInterface* Material = MediaPlate->GetCurrentMaterial())
+			{
+				Material->RecacheUniformExpressions(false);
+			}
+		}
 	}
 
 	// Set up sound component if we have one.
@@ -777,16 +785,7 @@ void UMediaPlateComponent::PostEditChangeProperty(FPropertyChangedEvent& Propert
 		{
 			MediaTextureTrackerObject->MipMapLODBias = MipMapBias;
 
-			// Note: Media texture bias automatically updated by UMediaPlateComponent::OnRegister().
-
-			// Update material sampler with new bias
-			if (AMediaPlate* MediaPlate = GetOwner<AMediaPlate>())
-			{
-				if (UMaterialInterface* Material = MediaPlate->GetCurrentMaterial())
-				{
-					Material->PostEditChange();
-				}
-			}
+			// Note: Media texture bias and material sampler automatically updated by UMediaPlateComponent::OnRegister().
 
 			RestartPlayer();
 		}
