MediaPlate: Fixed next/previous buttons not working.
#rnx

diff --git a/Engine/Plugins/Media/MediaPlate/Source/MediaPlate/Private/MediaPlateComponent.cpp b/Engine/Plugins/Media/MediaPlate/Source/MediaPlate/Private/MediaPlateComponent.cpp
--- a/Engine/Plugins/Media/MediaPlate/Source/MediaPlate/Private/MediaPlateComponent.cpp
+++ b/Engine/Plugins/Media/MediaPlate/Source/MediaPlate/Private/MediaPlateComponent.cpp
@@ -224,6 +224,27 @@ void UMediaPlateComponent::Open()
 	UpdateTicking();
 }
 
+bool UMediaPlateComponent::Next()
+{
+	bool bIsSuccessful = false;
+
+	// Do we have a playlist?
+	if ((MediaPlaylist != nullptr) && (MediaPlaylist->Num() > 1))
+	{
+		if (PlaylistIndex < MediaPlaylist->Num() - 1)
+		{
+			// Get the next media to play.
+			UMediaSource* NextSource = MediaPlaylist->GetNext(PlaylistIndex);
+			if (NextSource != nullptr)
+			{
+				bIsSuccessful = PlayMediaSource(NextSource);
+			}
+		}
+	}
+
+	return bIsSuccessful;
+}
+
 void UMediaPlateComponent::Play()
 {
 	if (MediaPlayer != nullptr)
@@ -242,6 +263,27 @@ void UMediaPlateComponent::Pause()
 	CurrentRate = 0.0f;
 }
 
+bool UMediaPlateComponent::Previous()
+{
+	bool bIsSuccessful = false;
+
+	// Do we have a playlist?
+	if ((MediaPlaylist != nullptr) && (MediaPlaylist->Num() > 1))
+	{
+		// Get the previous media to play.
+		if (PlaylistIndex > 0)
+		{
+			UMediaSource* NextSource = MediaPlaylist->GetPrevious(PlaylistIndex);
+			if (NextSource != nullptr)
+			{
+				bIsSuccessful = PlayMediaSource(NextSource);
+			}
+		}
+	}
+
+	return bIsSuccessful;
+}
+
 bool UMediaPlateComponent::Rewind()
 {
 	return Seek(FTimespan::Zero());
@@ -700,19 +742,7 @@ void UMediaPlateComponent::OnMediaEnd()
 {
 	StopClockSink();
 
-	// Do we have a playlist>
-	if ((MediaPlaylist != nullptr) && (MediaPlaylist->Num() > 1))
-	{
-		// Get the next media to play.
-		UMediaSource* NextSource = MediaPlaylist->GetNext(PlaylistIndex);
-		if (NextSource != nullptr)
-		{
-			if (PlaylistIndex != 0)
-			{
-				PlayMediaSource(NextSource);
-			}
-		}
-	}
+	Next();
 }
 
 #if WITH_EDITOR
diff --git a/Engine/Plugins/Media/MediaPlate/Source/MediaPlate/Public/MediaPlateComponent.h b/Engine/Plugins/Media/MediaPlate/Source/MediaPlate/Public/MediaPlateComponent.h
--- a/Engine/Plugins/Media/MediaPlate/Source/MediaPlate/Public/MediaPlateComponent.h
+++ b/Engine/Plugins/Media/MediaPlate/Source/MediaPlate/Public/MediaPlateComponent.h
@@ -58,6 +58,13 @@ public:
 	UFUNCTION(BlueprintCallable, Category = "Media|MediaPlateComponent")
 	void Open();
 
+	/**
+	 * Play the next item in the playlist.
+	 *
+	 * returns	True if it played something.
+	 */
+	bool Next();
+
 	/**
 	 * Call this to start playing.
 	 * Open must be called before this.
@@ -72,6 +79,13 @@ public:
 	UFUNCTION(BlueprintCallable, Category = "Media|MediaPlateComponent")
 	void Pause();
 
+	/**
+	 * Play the previous item in the playlist.
+	 *
+	 * returns	True if it played something.
+	 */
+	bool Previous();
+
 	/**
 	 * Rewinds the media to the beginning.
 	 *
diff --git a/Engine/Plugins/Media/MediaPlate/Source/MediaPlateEditor/Private/Toolkits/MediaPlateEditorToolkit.cpp b/Engine/Plugins/Media/MediaPlate/Source/MediaPlateEditor/Private/Toolkits/MediaPlateEditorToolkit.cpp
--- a/Engine/Plugins/Media/MediaPlate/Source/MediaPlateEditor/Private/Toolkits/MediaPlateEditorToolkit.cpp
+++ b/Engine/Plugins/Media/MediaPlate/Source/MediaPlateEditor/Private/Toolkits/MediaPlateEditorToolkit.cpp
@@ -262,8 +262,11 @@ void FMediaPlateEditorToolkit::BindCommands()
 
 	ToolkitCommands->MapAction(
 		Commands.NextMedia,
-		FExecuteAction::CreateLambda([this]{ MediaPlate->GetMediaPlayer()->Next(); }),
-		FCanExecuteAction::CreateLambda([this]{ return (MediaPlate->GetMediaPlayer()->GetPlaylistRef().Num() > 1); })
+		FExecuteAction::CreateLambda([this]{ MediaPlate->Next(); }),
+		FCanExecuteAction::CreateLambda([this]{
+			return (MediaPlate->MediaPlaylist != nullptr) &&
+				(MediaPlate->MediaPlaylist->Num() > 1);
+		})
 	);
 
 	ToolkitCommands->MapAction(
@@ -292,8 +295,11 @@ void FMediaPlateEditorToolkit::BindCommands()
 
 	ToolkitCommands->MapAction(
 		Commands.PreviousMedia,
-		FExecuteAction::CreateLambda([this]{ MediaPlate->GetMediaPlayer()->Previous(); }),
-		FCanExecuteAction::CreateLambda([this]{ return (MediaPlate->GetMediaPlayer()->GetPlaylistRef().Num() > 1); })
+		FExecuteAction::CreateLambda([this]{ MediaPlate->Previous(); }),
+		FCanExecuteAction::CreateLambda([this]{
+			return (MediaPlate->MediaPlaylist != nullptr) &&
+					(MediaPlate->MediaPlaylist->Num() > 1);
+		})
 	);
 
 	ToolkitCommands->MapAction(
