MediaPlate: Prevent invalidating the media texture on unrelated property changes.

#rnx

diff --git a/Engine/Plugins/Media/MediaPlate/Source/MediaPlate/Private/MediaPlateComponent.cpp b/Engine/Plugins/Media/MediaPlate/Source/MediaPlate/Private/MediaPlateComponent.cpp
--- a/Engine/Plugins/Media/MediaPlate/Source/MediaPlate/Private/MediaPlateComponent.cpp
+++ b/Engine/Plugins/Media/MediaPlate/Source/MediaPlate/Private/MediaPlateComponent.cpp
@@ -102,15 +102,31 @@ void UMediaPlateComponent::OnRegister()
 	// Set up media texture.
 	if (MediaTexture != nullptr)
 	{
-		MediaTexture->SetMipMapBias(MipMapBias);
-		MediaTexture->SetMediaPlayer(MediaPlayer);
-		MediaTexture->UpdateResource();
+		// Prevent media texture blackouts by only updating resource and material uniforms on relevant changes.
+		bool bApplyTextureMaterialUpdate = false;
 
-		if (AMediaPlate* MediaPlate = GetOwner<AMediaPlate>())
+		if (FMath::IsNearlyEqual(MediaTexture->GetMipMapBias(), MipMapBias) == false)
 		{
-			if (UMaterialInterface* Material = MediaPlate->GetCurrentMaterial())
+			MediaTexture->SetMipMapBias(MipMapBias);
+			bApplyTextureMaterialUpdate = true;
+		}
+
+		if (MediaTexture->GetMediaPlayer() != MediaPlayer.Get())
+		{
+			MediaTexture->SetMediaPlayer(MediaPlayer);
+			bApplyTextureMaterialUpdate = true;
+		}
+
+		if (bApplyTextureMaterialUpdate)
+		{
+			MediaTexture->UpdateResource();
+
+			if (AMediaPlate* MediaPlate = GetOwner<AMediaPlate>())
 			{
-				Material->RecacheUniformExpressions(false);
+				if (UMaterialInterface* Material = MediaPlate->GetCurrentMaterial())
+				{
+					Material->RecacheUniformExpressions(false);
+				}
 			}
 		}
 	}
