MediaPlate: Applies default material if override material appears to have been lost.
#rnx

diff --git a/Engine/Plugins/Media/MediaPlate/Source/MediaPlate/Private/MediaPlate.cpp b/Engine/Plugins/Media/MediaPlate/Source/MediaPlate/Private/MediaPlate.cpp
--- a/Engine/Plugins/Media/MediaPlate/Source/MediaPlate/Private/MediaPlate.cpp
+++ b/Engine/Plugins/Media/MediaPlate/Source/MediaPlate/Private/MediaPlate.cpp
@@ -83,6 +83,18 @@ void AMediaPlate::PostRegisterAllComponents()
 	Super::PostRegisterAllComponents();
 
 #if WITH_EDITOR
+	// If this media plate is in a Sequencer media track,
+	// and we go though a convert to spawnable/possessible,
+	// then a multi user client will not receive the correct material
+	// as multi user cannot send materials that are not assets.
+	// So if we have an override material but its nullptr, then just use the default material.
+	if ((StaticMeshComponent != nullptr) &&
+		(StaticMeshComponent->GetNumOverrideMaterials() > 0) &&
+		(StaticMeshComponent->OverrideMaterials[0] == nullptr))
+	{
+		UseDefaultMaterial();		
+	}
+	
 	AddAssetUserData();
 #endif // WITH_EDITOR
 }
@@ -123,63 +135,73 @@ void AMediaPlate::ApplyCurrentMaterial()
 
 void AMediaPlate::ApplyMaterial(UMaterialInterface* Material)
 {
-	if (GEditor != nullptr && Material != nullptr && StaticMeshComponent != nullptr)
+	if (Material != nullptr && StaticMeshComponent != nullptr)
 	{
-		UMaterialInstanceDynamic* MID = Cast<UMaterialInstanceDynamic>(Material);
-		UMaterialInterface* Result = nullptr;
-
-		// See if we can modify this material.
-		bool bCanModify = true;
-		FMediaPlateModule* MediaPlateModule = FModuleManager::GetModulePtr<FMediaPlateModule>("MediaPlate");
-		if (MediaPlateModule != nullptr)
-		{
-			MediaPlateModule->OnMediaPlateApplyMaterial.Broadcast(Material, this, bCanModify);
-		}
-		
-		if (bCanModify == false)
+		if (GEditor == nullptr)
 		{
-			LastMaterial = Material;
-		}
-		else if (MID != nullptr)
-		{
-			MID->SetTextureParameterValue(MediaTextureName, MediaPlateComponent->GetMediaTexture());
+			UMaterialInstanceDynamic* MaterialDynamic = StaticMeshComponent->CreateAndSetMaterialInstanceDynamicFromMaterial(0, Material);
+			MaterialDynamic->SetTextureParameterValue(MediaTextureName, MediaPlateComponent->GetMediaTexture());
 
-			Result = MID;
+			LastMaterial = MaterialDynamic;
 		}
 		else
 		{
-			// Change M_ to MI_ in material name and then generate a unique one.
-			FString MaterialName = Material->GetName();
-			if (MaterialName.StartsWith(TEXT("M_")))
+			UMaterialInstanceDynamic* MID = Cast<UMaterialInstanceDynamic>(Material);
+			UMaterialInterface* Result = nullptr;
+
+			// See if we can modify this material.
+			bool bCanModify = true;
+			FMediaPlateModule* MediaPlateModule = FModuleManager::GetModulePtr<FMediaPlateModule>("MediaPlate");
+			if (MediaPlateModule != nullptr)
 			{
-				MaterialName.InsertAt(1, TEXT("I"));
+				MediaPlateModule->OnMediaPlateApplyMaterial.Broadcast(Material, this, bCanModify);
 			}
-			FName MaterialUniqueName = MakeUniqueObjectName(StaticMeshComponent, UMaterialInstanceConstant::StaticClass(),
-				FName(*MaterialName));
-
-			// Create instance.
-			UMaterialInstanceConstant* MaterialInstance =
-				NewObject<UMaterialInstanceConstant>(StaticMeshComponent, MaterialUniqueName);
-			MaterialInstance->SetParentEditorOnly(Material);
-			MaterialInstance->CopyMaterialUniformParametersEditorOnly(Material);
-			MaterialInstance->SetTextureParameterValueEditorOnly(
-				FMaterialParameterInfo(MediaTextureName),
-				MediaPlateComponent->GetMediaTexture());
-			MaterialInstance->PostEditChange();
-
-			// We force call post-load to indirectly call UpdateParameters() (for integration with VPUtilities plugin).
-			MaterialInstance->PostLoad();
-
-			Result = MaterialInstance;
-		}
 
-		// Update static mesh.
-		if (Result != nullptr)
-		{
-			StaticMeshComponent->Modify();
-			StaticMeshComponent->SetMaterial(0, Result);
+			if (bCanModify == false)
+			{
+				LastMaterial = Material;
+			}
+			else if (MID != nullptr)
+			{
+				MID->SetTextureParameterValue(MediaTextureName, MediaPlateComponent->GetMediaTexture());
+
+				Result = MID;
+			}
+			else
+			{
+				// Change M_ to MI_ in material name and then generate a unique one.
+				FString MaterialName = Material->GetName();
+				if (MaterialName.StartsWith(TEXT("M_")))
+				{
+					MaterialName.InsertAt(1, TEXT("I"));
+				}
+				FName MaterialUniqueName = MakeUniqueObjectName(StaticMeshComponent, UMaterialInstanceConstant::StaticClass(),
+					FName(*MaterialName));
+
+				// Create instance.
+				UMaterialInstanceConstant* MaterialInstance =
+					NewObject<UMaterialInstanceConstant>(StaticMeshComponent, MaterialUniqueName);
+				MaterialInstance->SetParentEditorOnly(Material);
+				MaterialInstance->CopyMaterialUniformParametersEditorOnly(Material);
+				MaterialInstance->SetTextureParameterValueEditorOnly(
+					FMaterialParameterInfo(MediaTextureName),
+					MediaPlateComponent->GetMediaTexture());
+				MaterialInstance->PostEditChange();
+
+				// We force call post-load to indirectly call UpdateParameters() (for integration with VPUtilities plugin).
+				MaterialInstance->PostLoad();
+
+				Result = MaterialInstance;
+			}
 
-			LastMaterial = Result;
+			// Update static mesh.
+			if (Result != nullptr)
+			{
+				StaticMeshComponent->Modify();
+				StaticMeshComponent->SetMaterial(0, Result);
+
+				LastMaterial = Result;
+			}
 		}
 	}
 }
