MediaPlate: Auto aspect ratio can now be toggled during playback.

diff --git a/Engine/Plugins/Media/MediaPlate/Source/MediaPlate/Private/MediaPlateComponent.cpp b/Engine/Plugins/Media/MediaPlate/Source/MediaPlate/Private/MediaPlateComponent.cpp
--- a/Engine/Plugins/Media/MediaPlate/Source/MediaPlate/Private/MediaPlateComponent.cpp
+++ b/Engine/Plugins/Media/MediaPlate/Source/MediaPlate/Private/MediaPlateComponent.cpp
@@ -372,6 +372,15 @@ void UMediaPlateComponent::SetPlayOnlyWhenVisible(bool bInPlayOnlyWhenVisible)
 	PlayOnlyWhenVisibleChanged();
 }
 
+void UMediaPlateComponent::SetIsAspectRatioAuto(bool bInIsAspectRatioAuto)
+{
+	if (bIsAspectRatioAuto != bInIsAspectRatioAuto)
+	{
+		bIsAspectRatioAuto = bInIsAspectRatioAuto;
+		TryActivateAspectRatioAuto();
+	}
+}
+
 void UMediaPlateComponent::PlayOnlyWhenVisibleChanged()
 {
 	// If we are turning off PlayOnlyWhenVisible then make sure we are playing.
@@ -438,21 +447,7 @@ bool UMediaPlateComponent::PlayMediaSource(UMediaSource* InMediaSource)
 			// Did we play anything?
 			if (bIsPlaying)
 			{
-				// Are we using automatic aspect ratio?
-				if ((bIsAspectRatioAuto) &&
-					(VisibleMipsTilesCalculations == EMediaTextureVisibleMipsTiles::Plane))
-				{
-					// Start the clock sink so we can tick.
-					IMediaModule* MediaModule = FModuleManager::LoadModulePtr<IMediaModule>("Media");
-					if (MediaModule != nullptr)
-					{
-						if (ClockSink.IsValid() == false)
-						{
-							ClockSink = MakeShared<FMediaComponentClockSink, ESPMode::ThreadSafe>(this);
-						}
-						MediaModule->GetClock().AddSink(ClockSink.ToSharedRef());
-					}
-				}
+				TryActivateAspectRatioAuto();
 			}
 		}
 	}
@@ -460,6 +455,28 @@ bool UMediaPlateComponent::PlayMediaSource(UMediaSource* InMediaSource)
 	return bIsPlaying;
 }
 
+void UMediaPlateComponent::TryActivateAspectRatioAuto()
+{
+	if ((MediaPlayer != nullptr) && (MediaPlayer->IsClosed() == false))
+	{
+		// Are we using automatic aspect ratio?
+		if ((bIsAspectRatioAuto) &&
+			(VisibleMipsTilesCalculations == EMediaTextureVisibleMipsTiles::Plane))
+		{
+			// Start the clock sink so we can tick.
+			IMediaModule* MediaModule = FModuleManager::LoadModulePtr<IMediaModule>("Media");
+			if (MediaModule != nullptr)
+			{
+				if (ClockSink.IsValid() == false)
+				{
+					ClockSink = MakeShared<FMediaComponentClockSink, ESPMode::ThreadSafe>(this);
+				}
+				MediaModule->GetClock().AddSink(ClockSink.ToSharedRef());
+			}
+		}
+	}
+}
+
 
 float UMediaPlateComponent::GetAspectRatio()
 {
diff --git a/Engine/Plugins/Media/MediaPlate/Source/MediaPlate/Public/MediaPlateComponent.h b/Engine/Plugins/Media/MediaPlate/Source/MediaPlate/Public/MediaPlateComponent.h
--- a/Engine/Plugins/Media/MediaPlate/Source/MediaPlate/Public/MediaPlateComponent.h
+++ b/Engine/Plugins/Media/MediaPlate/Source/MediaPlate/Public/MediaPlateComponent.h
@@ -146,10 +146,6 @@ public:
 	UPROPERTY(EditAnywhere, BlueprintReadWrite, Category = "MediaPlate", meta = (ClampMin = "0.0"))
 	float StartTime = 0.0f;
 
-	/** If true then set the aspect ratio automatically based on the media. */
-	UPROPERTY(BlueprintReadWrite, Category = "MediaPlate")
-	bool bIsAspectRatioAuto = true;
-
 	/** Holds the component to play sound. */
 	UPROPERTY(EditAnywhere, AdvancedDisplay, Category = MediaPlate)
 	TObjectPtr<UMediaSoundComponent> SoundComponent;
@@ -195,6 +191,18 @@ public:
 	 */
 	void SetAspectRatio(float AspectRatio);
 
+	/**
+	 * Gets whether automatic aspect ratio is enabled.
+	 */
+	UFUNCTION(BlueprintCallable, Category = "Media|MediaPlateComponent")
+	bool GetIsAspectRatioAuto() const { return bIsAspectRatioAuto; }
+
+	/**
+	 * Sets whether automatic aspect ratio is enabled.
+	 */
+	UFUNCTION(BlueprintCallable, Category = "Media|MediaPlateComponent")
+	void SetIsAspectRatioAuto(bool bInIsAspectRatioAuto);
+
 	/**
 	 * Call this to get the aspect ratio of the screen.
 	 */
@@ -278,6 +286,10 @@ private:
 	UPROPERTY(EditAnywhere, AdvancedDisplay, Category = "MediaPlate", meta = (AllowPrivateAccess = true, UIMin = "-16.0", UIMax = "15.99"))
 	float MipMapBias = 0.0f;
 
+	/** If true then set the aspect ratio automatically based on the media. */
+	UPROPERTY(Blueprintgetter = GetIsAspectRatioAuto, BlueprintSetter = SetIsAspectRatioAuto, Category = "MediaPlate", meta = (AllowPrivateAccess = true))
+	bool bIsAspectRatioAuto = true;
+
 	/** If > 0, then this is the aspect ratio of our screen and 
 	 * letterboxes will be added if the media is smaller than the screen. */
 	UPROPERTY()
@@ -318,6 +330,12 @@ private:
 	 */
 	bool PlayMediaSource(UMediaSource* InMediaSource);
 
+	/**
+	 * If the player is currently active, then this will set the aspect ratio
+	 * according to the media.
+	 */
+	void TryActivateAspectRatioAuto();
+
 	/**
 	 * Stops the clock sink so we no longer tick.
 	 */
diff --git a/Engine/Plugins/Media/MediaPlate/Source/MediaPlateEditor/Private/MediaPlateCustomization.cpp b/Engine/Plugins/Media/MediaPlate/Source/MediaPlateEditor/Private/MediaPlateCustomization.cpp
--- a/Engine/Plugins/Media/MediaPlate/Source/MediaPlateEditor/Private/MediaPlateCustomization.cpp
+++ b/Engine/Plugins/Media/MediaPlate/Source/MediaPlateEditor/Private/MediaPlateCustomization.cpp
@@ -803,7 +803,7 @@ ECheckBoxState FMediaPlateCustomization::IsAspectRatioAuto() const
 		UMediaPlateComponent* MediaPlate = MediaPlatePtr.Get();
 		if (MediaPlate != nullptr)
 		{
-			ECheckBoxState NewState = MediaPlate->bIsAspectRatioAuto ? ECheckBoxState::Checked :
+			ECheckBoxState NewState = MediaPlate->GetIsAspectRatioAuto() ? ECheckBoxState::Checked :
 				ECheckBoxState::Unchecked;
 			if (State == ECheckBoxState::Undetermined)
 			{
@@ -831,7 +831,7 @@ void FMediaPlateCustomization::SetIsAspectRatioAuto(ECheckBoxState State)
 		UMediaPlateComponent* MediaPlate = MediaPlatePtr.Get();
 		if (MediaPlate != nullptr)
 		{
-			MediaPlate->bIsAspectRatioAuto = bEnable;
+			MediaPlate->SetIsAspectRatioAuto(bEnable);
 		}
 	}
 }
