Fix crash when exposing property on media plate actor

diff --git a/Engine/Plugins/Media/MediaPlate/Source/MediaPlateEditor/Private/MediaPlateCustomization.cpp b/Engine/Plugins/Media/MediaPlate/Source/MediaPlateEditor/Private/MediaPlateCustomization.cpp
--- a/Engine/Plugins/Media/MediaPlate/Source/MediaPlateEditor/Private/MediaPlateCustomization.cpp
+++ b/Engine/Plugins/Media/MediaPlate/Source/MediaPlateEditor/Private/MediaPlateCustomization.cpp
@@ -46,9 +46,13 @@ FMediaPlateCustomization::~FMediaPlateCustomization()
 void FMediaPlateCustomization::CustomizeDetails(IDetailLayoutBuilder& DetailBuilder)
 {
 	// Is this the media plate editor window?
-	const IDetailsView* DetailsView = DetailBuilder.GetDetailsView();
-	TSharedPtr<FTabManager> HostTabManager = DetailsView->GetHostTabManager();
-	bool bIsMediaPlateWindow = (HostTabManager.IsValid() == false);
+	bool bIsMediaPlateWindow = false;
+
+	if (const IDetailsView* DetailsView = DetailBuilder.GetDetailsView())
+	{
+		TSharedPtr<FTabManager> HostTabManager = DetailsView->GetHostTabManager();
+		bIsMediaPlateWindow = (HostTabManager.IsValid() == false);
+	}
 
 	// Get style.
 	const ISlateStyle* Style = nullptr;
