MediaPlate: Fixed crash when playing a media with duration 0.

diff --git a/Engine/Plugins/Media/MediaPlate/Source/MediaPlate/Private/MediaPlateComponent.cpp b/Engine/Plugins/Media/MediaPlate/Source/MediaPlate/Private/MediaPlateComponent.cpp
--- a/Engine/Plugins/Media/MediaPlate/Source/MediaPlate/Private/MediaPlateComponent.cpp
+++ b/Engine/Plugins/Media/MediaPlate/Source/MediaPlate/Private/MediaPlateComponent.cpp
@@ -657,7 +657,7 @@ FTimespan UMediaPlateComponent::GetResumeTime()
 		
 		// Are we over the length of the media?
 		FTimespan MediaDuration = MediaPlayer->GetDuration();
-		if (PlayerTime > MediaDuration)
+		if ((PlayerTime > MediaDuration) && (MediaDuration > FTimespan::Zero()))
 		{
 			bool bIsPlaylist = (MediaPlaylist != nullptr) && (MediaPlaylist->Num() > 1);
 			if ((bLoop) && (bIsPlaylist == false))
