MediaPlate: Fixed audio not working for multiple media plates.

diff --git a/Engine/Plugins/Media/MediaPlate/Source/MediaPlateEditor/Private/MediaPlateEditorModule.cpp b/Engine/Plugins/Media/MediaPlate/Source/MediaPlateEditor/Private/MediaPlateEditorModule.cpp
--- a/Engine/Plugins/Media/MediaPlate/Source/MediaPlateEditor/Private/MediaPlateEditorModule.cpp
+++ b/Engine/Plugins/Media/MediaPlate/Source/MediaPlateEditor/Private/MediaPlateEditorModule.cpp
@@ -104,7 +104,6 @@ void FMediaPlateEditorModule::Tick(float DeltaTime)
 			if (MediaPlate->IsMediaPlatePlaying())
 			{
 				bIsMediaPlatePlaying = true;
-				break;
 			}
 		}
 
