MediaPlate: Fixed ActiveMediaPlates getting duplicate entries.

diff --git a/Engine/Plugins/Media/MediaPlate/Source/MediaPlateEditor/Private/MediaPlateEditorModule.cpp b/Engine/Plugins/Media/MediaPlate/Source/MediaPlateEditor/Private/MediaPlateEditorModule.cpp
--- a/Engine/Plugins/Media/MediaPlate/Source/MediaPlateEditor/Private/MediaPlateEditorModule.cpp
+++ b/Engine/Plugins/Media/MediaPlate/Source/MediaPlateEditor/Private/MediaPlateEditorModule.cpp
@@ -123,7 +123,7 @@ void FMediaPlateEditorModule::MediaPlateStartedPlayback(TObjectPtr<UMediaPlateCo
 {
 	if (MediaPlate != nullptr)
 	{
-		ActiveMediaPlates.Add(MediaPlate);
+		ActiveMediaPlates.AddUnique(MediaPlate);
 	}
 }
 
