MediaPlate: Fixed looping not working for playlists.

diff --git a/Engine/Plugins/Media/MediaPlate/Source/MediaPlate/Private/MediaPlateComponent.cpp b/Engine/Plugins/Media/MediaPlate/Source/MediaPlate/Private/MediaPlateComponent.cpp
--- a/Engine/Plugins/Media/MediaPlate/Source/MediaPlate/Private/MediaPlateComponent.cpp
+++ b/Engine/Plugins/Media/MediaPlate/Source/MediaPlate/Private/MediaPlateComponent.cpp
@@ -256,7 +256,7 @@ bool UMediaPlateComponent::Next()
 	// Do we have a playlist?
 	if ((MediaPlaylist != nullptr) && (MediaPlaylist->Num() > 1))
 	{
-		if (PlaylistIndex < MediaPlaylist->Num() - 1)
+		if ((PlaylistIndex < MediaPlaylist->Num() - 1) || (bLoop))
 		{
 			// Get the next media to play.
 			UMediaSource* NextSource = MediaPlaylist->GetNext(PlaylistIndex);
@@ -435,13 +435,15 @@ bool UMediaPlateComponent::PlayMediaSource(UMediaSource* InMediaSource)
 		// Set media options.
 		if (MediaPlayer != nullptr)
 		{
+			bool bIsPlaylist = (MediaPlaylist != nullptr) && (MediaPlaylist->Num() > 1);
+
 			// Play the source.
 			FMediaPlayerOptions Options;
 			Options.SeekTime = FTimespan::FromSeconds(StartTime);
 			Options.PlayOnOpen = bPlayOnOpen ? EMediaPlayerOptionBooleanOverride::Enabled :
 				EMediaPlayerOptionBooleanOverride::Disabled;
-			Options.Loop = bLoop ? EMediaPlayerOptionBooleanOverride::Enabled :
-				EMediaPlayerOptionBooleanOverride::Disabled;
+			Options.Loop = (bLoop && (bIsPlaylist == false)) ?
+				EMediaPlayerOptionBooleanOverride::Enabled : EMediaPlayerOptionBooleanOverride::Disabled;
 			bIsPlaying = MediaPlayer->OpenSourceWithOptions(InMediaSource, Options);
 
 			// Did we play anything?
