MediaPlate: Fixed subsequent tracks in the playlist not starting when play on open is false.

diff --git a/Engine/Plugins/Media/MediaPlate/Source/MediaPlate/Private/MediaPlateComponent.cpp b/Engine/Plugins/Media/MediaPlate/Source/MediaPlate/Private/MediaPlateComponent.cpp
--- a/Engine/Plugins/Media/MediaPlate/Source/MediaPlate/Private/MediaPlateComponent.cpp
+++ b/Engine/Plugins/Media/MediaPlate/Source/MediaPlate/Private/MediaPlateComponent.cpp
@@ -231,7 +231,7 @@ void UMediaPlateComponent::Open()
 			{
 				MediaSource = MediaPlaylist->Get(0);
 			}
-			bIsPlaying = PlayMediaSource(MediaSource);
+			bIsPlaying = PlayMediaSource(MediaSource, bPlayOnOpen);
 		}
 
 		// Did anything play?
@@ -262,7 +262,7 @@ bool UMediaPlateComponent::Next()
 			UMediaSource* NextSource = MediaPlaylist->GetNext(PlaylistIndex);
 			if (NextSource != nullptr)
 			{
-				bIsSuccessful = PlayMediaSource(NextSource);
+				bIsSuccessful = PlayMediaSource(NextSource, true);
 			}
 		}
 	}
@@ -301,7 +301,7 @@ bool UMediaPlateComponent::Previous()
 			UMediaSource* NextSource = MediaPlaylist->GetPrevious(PlaylistIndex);
 			if (NextSource != nullptr)
 			{
-				bIsSuccessful = PlayMediaSource(NextSource);
+				bIsSuccessful = PlayMediaSource(NextSource, true);
 			}
 		}
 	}
@@ -423,7 +423,7 @@ void UMediaPlateComponent::UnregisterWithMediaTextureTracker()
 	}
 }
 
-bool UMediaPlateComponent::PlayMediaSource(UMediaSource* InMediaSource)
+bool UMediaPlateComponent::PlayMediaSource(UMediaSource* InMediaSource, bool bInPlayOnOpen)
 {
 	bool bIsPlaying = false;
 
@@ -440,7 +440,7 @@ bool UMediaPlateComponent::PlayMediaSource(UMediaSource* InMediaSource)
 			// Play the source.
 			FMediaPlayerOptions Options;
 			Options.SeekTime = FTimespan::FromSeconds(StartTime);
-			Options.PlayOnOpen = bPlayOnOpen ? EMediaPlayerOptionBooleanOverride::Enabled :
+			Options.PlayOnOpen = bInPlayOnOpen ? EMediaPlayerOptionBooleanOverride::Enabled :
 				EMediaPlayerOptionBooleanOverride::Disabled;
 			Options.Loop = (bLoop && (bIsPlaylist == false)) ?
 				EMediaPlayerOptionBooleanOverride::Enabled : EMediaPlayerOptionBooleanOverride::Disabled;
diff --git a/Engine/Plugins/Media/MediaPlate/Source/MediaPlate/Public/MediaPlateComponent.h b/Engine/Plugins/Media/MediaPlate/Source/MediaPlate/Public/MediaPlateComponent.h
--- a/Engine/Plugins/Media/MediaPlate/Source/MediaPlate/Public/MediaPlateComponent.h
+++ b/Engine/Plugins/Media/MediaPlate/Source/MediaPlate/Public/MediaPlateComponent.h
@@ -327,9 +327,10 @@ private:
 	 * Plays a media source.
 	 * 
 	 * @param	InMediaSource		Media source to play.
+	 * @param	bInPlayOnOpen		True to play, false to just open.
 	 * @return	True if we played anything.
 	 */
-	bool PlayMediaSource(UMediaSource* InMediaSource);
+	bool PlayMediaSource(UMediaSource* InMediaSource, bool bInPlayOnOpen);
 
 	/**
 	 * If the player is currently active, then this will set the aspect ratio
