MediaPlate: Removed Instanced from playlist as it was preventing the playlist reference from being saved, and Instanced is no longer necessary as Sequencer is no longer directly linking to the media sources in the media plate.

diff --git a/Engine/Plugins/Media/MediaPlate/Source/MediaPlate/Public/MediaPlateComponent.h b/Engine/Plugins/Media/MediaPlate/Source/MediaPlate/Public/MediaPlateComponent.h
--- a/Engine/Plugins/Media/MediaPlate/Source/MediaPlate/Public/MediaPlateComponent.h
+++ b/Engine/Plugins/Media/MediaPlate/Source/MediaPlate/Public/MediaPlateComponent.h
@@ -158,7 +158,7 @@ public:
 	TArray<TObjectPtr<UStaticMeshComponent>> Letterboxes;
 
 	/** What media playlist to play. */
-	UPROPERTY(Instanced, BlueprintReadWrite, Category = "MediaPlate")
+	UPROPERTY(BlueprintReadWrite, Category = "MediaPlate")
 	TObjectPtr<UMediaPlaylist> MediaPlaylist;
 
 	/** The current index of the source in the play list being played. */
