MediaPlate: Added support for multiple textures for cross fades.

diff --git a/Engine/Plugins/Media/MediaPlate/Source/MediaPlate/Private/MediaPlate.cpp b/Engine/Plugins/Media/MediaPlate/Source/MediaPlate/Private/MediaPlate.cpp
--- a/Engine/Plugins/Media/MediaPlate/Source/MediaPlate/Private/MediaPlate.cpp
+++ b/Engine/Plugins/Media/MediaPlate/Source/MediaPlate/Private/MediaPlate.cpp
@@ -141,8 +141,7 @@ void AMediaPlate::ApplyMaterial(UMaterialInterface* Material)
 		if (GEditor == nullptr)
 		{
 			UMaterialInstanceDynamic* MaterialDynamic = StaticMeshComponent->CreateAndSetMaterialInstanceDynamicFromMaterial(0, Material);
-			MaterialDynamic->SetTextureParameterValue(MediaTextureName, MediaPlateComponent->GetMediaTexture());
-
+			SetMIDParameters(MaterialDynamic);
 			LastMaterial = MaterialDynamic;
 		}
 		else
@@ -160,16 +159,18 @@ void AMediaPlate::ApplyMaterial(UMaterialInterface* Material)
 
 			if (bCanModify == false)
 			{
+				MediaPlateComponent->SetNumberOfTextures(1);
 				LastMaterial = Material;
 			}
 			else if (MID != nullptr)
 			{
-				MID->SetTextureParameterValue(MediaTextureName, MediaPlateComponent->GetMediaTexture());
-
+				SetMIDParameters(MID);
 				Result = MID;
 			}
 			else
 			{
+				MediaPlateComponent->SetNumberOfTextures(1);
+
 				// Change M_ to MI_ in material name and then generate a unique one.
 				FString MaterialName = Material->GetName();
 				if (MaterialName.StartsWith(TEXT("M_")))
@@ -207,6 +208,35 @@ void AMediaPlate::ApplyMaterial(UMaterialInterface* Material)
 	}
 }
 
+void AMediaPlate::SetMIDParameters(UMaterialInstanceDynamic* InMaterial)
+{
+	InMaterial->SetTextureParameterValue(MediaTextureName, MediaPlateComponent->GetMediaTexture());
+	
+	int32 NumTextures = 0;
+	FString MediaTextureString = MediaTextureName.Resolve().ToString();
+	for (const struct FTextureParameterValue& Param : InMaterial->TextureParameterValues)
+	{
+		FString Name = Param.ParameterInfo.Name.ToString();
+		if (Name.StartsWith(MediaTextureString))
+		{
+			NumTextures++;
+		}
+	}
+	MediaPlateComponent->SetNumberOfTextures(NumTextures);
+
+	for (int32 Index = 0; Index < NumTextures; Index++)
+	{
+		FString NameString = MediaTextureString;
+		if (Index != 0)
+		{
+			NameString.AppendInt(Index);
+		}
+		FName TextureParameterName = FName(*NameString);
+		InMaterial->SetTextureParameterValue(TextureParameterName,
+			MediaPlateComponent->GetMediaTexture(Index));
+	}
+}
+
 void AMediaPlate::OnPreSaveWorld(UWorld* InWorld, FObjectPreSaveContext ObjectSaveContext)
 {
 	// We need to remove our asset user data before saving, as we do not need to save it out
diff --git a/Engine/Plugins/Media/MediaPlate/Source/MediaPlate/Private/MediaPlateComponent.cpp b/Engine/Plugins/Media/MediaPlate/Source/MediaPlate/Private/MediaPlateComponent.cpp
--- a/Engine/Plugins/Media/MediaPlate/Source/MediaPlate/Private/MediaPlateComponent.cpp
+++ b/Engine/Plugins/Media/MediaPlate/Source/MediaPlate/Private/MediaPlateComponent.cpp
@@ -8,6 +8,7 @@
 #include "IMediaClockSink.h"
 #include "IMediaModule.h"
 #include "Materials/Material.h"
+#include "Materials/MaterialInstanceDynamic.h"
 #include "MediaComponent.h"
 #include "MediaPlateModule.h"
 #include "MediaPlate.h"
@@ -23,6 +24,17 @@
 #define LOCTEXT_NAMESPACE "MediaPlate"
 
 
+namespace UE::MediaPlateComponent
+{
+	enum class ESetUpTexturesFlags
+	{
+		None,
+		AllowSetPlayer = 0x1,
+		ForceUpdateResource = 0x2,
+	};
+	ENUM_CLASS_FLAGS(ESetUpTexturesFlags);
+};
+
 /**
  * Media clock sink for media textures.
  */
@@ -83,10 +95,9 @@ void UMediaPlateComponent::OnRegister()
 	Super::OnRegister();
 
 	// Create media texture if we don't have one.
-	if (MediaTexture == nullptr)
+	if (MediaTextures.Num() == 0)
 	{
-		MediaTexture = NewObject<UMediaTexture>(this);
-		MediaTexture->NewStyleOutput = true;
+		SetNumberOfTextures(1);
 	}
 
 	// Create media player if we don't have one.
@@ -100,36 +111,7 @@ void UMediaPlateComponent::OnRegister()
 	MediaPlayer->OnEndReached.AddUniqueDynamic(this, &UMediaPlateComponent::OnMediaEnd);
 
 	// Set up media texture.
-	if (MediaTexture != nullptr)
-	{
-		// Prevent media texture blackouts by only updating resource and material uniforms on relevant changes.
-		bool bApplyTextureMaterialUpdate = false;
-
-		if (FMath::IsNearlyEqual(MediaTexture->GetMipMapBias(), MipMapBias) == false)
-		{
-			MediaTexture->SetMipMapBias(MipMapBias);
-			bApplyTextureMaterialUpdate = true;
-		}
-
-		if (MediaTexture->GetMediaPlayer() != MediaPlayer.Get())
-		{
-			MediaTexture->SetMediaPlayer(MediaPlayer);
-			bApplyTextureMaterialUpdate = true;
-		}
-
-		if (bApplyTextureMaterialUpdate)
-		{
-			MediaTexture->UpdateResource();
-
-			if (AMediaPlate* MediaPlate = GetOwner<AMediaPlate>())
-			{
-				if (UMaterialInterface* Material = MediaPlate->GetCurrentMaterial())
-				{
-					Material->RecacheUniformExpressions(false);
-				}
-			}
-		}
-	}
+	SetUpTextures(UE::MediaPlateComponent::ESetUpTexturesFlags::AllowSetPlayer);
 
 	// Set up sound component if we have one.
 	if (SoundComponent != nullptr)
@@ -211,8 +193,17 @@ UMediaPlayer* UMediaPlateComponent::GetMediaPlayer()
 	return MediaPlayer;
 }
 
-UMediaTexture* UMediaPlateComponent::GetMediaTexture()
+UMediaTexture* UMediaPlateComponent::GetMediaTexture(int32 Index)
 {
+	UMediaTexture* MediaTexture = nullptr;
+	if ((Index >= 0) && (Index < MediaTextures.Num()))
+	{
+		MediaTexture = MediaTextures[Index];
+	}
+	else
+	{
+		UE_LOG(LogMediaPlate, Warning, TEXT("Texture index %d not available."), Index);
+	}
 	return MediaTexture;
 }
 
@@ -411,11 +402,14 @@ void UMediaPlateComponent::RegisterWithMediaTextureTracker()
 	MediaTextureTrackerObject->VisibleMipsTilesCalculations = VisibleMipsTilesCalculations;
 	MediaTextureTrackerObject->MeshRange = MeshRange;
 
-	// Add our texture.
-	if (MediaTexture != nullptr)
+	// Add our textures.
+	FMediaTextureTracker& MediaTextureTracker = FMediaTextureTracker::Get();
+	for (UMediaTexture* MediaTexture : MediaTextures)
 	{
-		FMediaTextureTracker& MediaTextureTracker = FMediaTextureTracker::Get();
-		MediaTextureTracker.RegisterTexture(MediaTextureTrackerObject, MediaTexture);
+		if (MediaTexture != nullptr)
+		{
+			MediaTextureTracker.RegisterTexture(MediaTextureTrackerObject, MediaTexture);
+		}
 	}
 }
 
@@ -425,7 +419,10 @@ void UMediaPlateComponent::UnregisterWithMediaTextureTracker()
 	if (MediaTextureTrackerObject != nullptr)
 	{
 		FMediaTextureTracker& MediaTextureTracker = FMediaTextureTracker::Get();
-		MediaTextureTracker.UnregisterTexture(MediaTextureTrackerObject, MediaTexture);
+		for (UMediaTexture* MediaTexture : MediaTextures)
+		{
+			MediaTextureTracker.UnregisterTexture(MediaTextureTrackerObject, MediaTexture);
+		}
 	}
 }
 
@@ -539,6 +536,36 @@ void UMediaPlateComponent::SetLetterboxAspectRatio(float AspectRatio)
 	UpdateLetterboxes();
 }
 
+void UMediaPlateComponent::SetNumberOfTextures(int32 NumTextures)
+{
+	if (MediaTextures.Num() != NumTextures)
+	{
+		if (IsRegistered())
+		{
+			UnregisterWithMediaTextureTracker();
+		}
+		if (MediaTextures.Num() > NumTextures)
+		{
+			MediaTextures.RemoveAt(NumTextures, MediaTextures.Num() - NumTextures);
+		}
+		else
+		{
+			while (MediaTextures.Num() < NumTextures)
+			{
+				UMediaTexture* MediaTexture = NewObject<UMediaTexture>(this);
+				MediaTexture->NewStyleOutput = true;
+				MediaTextures.Add(MediaTexture);
+			}
+		}
+
+		SetUpTextures(UE::MediaPlateComponent::ESetUpTexturesFlags::ForceUpdateResource);
+		if (IsRegistered())
+		{
+			RegisterWithMediaTextureTracker();
+		}
+	}
+}
+
 void UMediaPlateComponent::TickOutput()
 {
 	if (MediaPlayer != nullptr)
@@ -606,8 +633,7 @@ void UMediaPlateComponent::ProxyClose()
 
 UMediaTexture* UMediaPlateComponent::ProxyGetMediaTexture(int32 TextureIndex)
 {
-	UMediaTexture* OutTexture = MediaTexture;
-	return OutTexture;
+	return GetMediaTexture(TextureIndex);
 }
 
 bool UMediaPlateComponent::ProxyIsPlaylistIndexPlaying(int32 Index) const
@@ -620,6 +646,24 @@ void UMediaPlateComponent::ProxySetPlayOnOpen(bool bInPlayOnOpen)
 	bPlayOnOpen = bInPlayOnOpen;
 }
 
+void UMediaPlateComponent::ProxySetTextureBlend(int32 TextureIndex, float Blend)
+{
+	if (AMediaPlate* MediaPlate = GetOwner<AMediaPlate>())
+	{
+		if (UMaterialInterface* Material = MediaPlate->GetCurrentMaterial())
+		{
+			UMaterialInstanceDynamic* MID = Cast<UMaterialInstanceDynamic>(Material);
+			if (MID != nullptr)
+			{
+				static const FString BaseBlendName = TEXT("Blend");
+				FString BlendName = BaseBlendName;
+				BlendName.AppendInt(TextureIndex);
+				MID->SetScalarParameterValue(FName(*BlendName), Blend);
+			}
+		}
+	}
+}
+
 void UMediaPlateComponent::RestartPlayer()
 {
 	if (MediaPlayer != nullptr)
@@ -832,6 +876,49 @@ void UMediaPlateComponent::OnMediaEnd()
 	Next();
 }
 
+void UMediaPlateComponent::SetUpTextures(UE::MediaPlateComponent::ESetUpTexturesFlags Flags)
+{
+	// Prevent media texture blackouts by only updating resource and material uniforms on relevant changes.
+	bool bApplyMaterialUpdate = false;
+	for (UMediaTexture* MediaTexture : MediaTextures)
+	{
+		if (MediaTexture != nullptr)
+		{
+			bool bApplyTextureUpdate = false;
+
+			if (FMath::IsNearlyEqual(MediaTexture->GetMipMapBias(), MipMapBias) == false)
+			{
+				MediaTexture->SetMipMapBias(MipMapBias);
+				bApplyTextureUpdate = true;
+				bApplyMaterialUpdate = true;
+			}
+
+			if ((EnumHasAllFlags(Flags, UE::MediaPlateComponent::ESetUpTexturesFlags::AllowSetPlayer)) &&
+				(MediaTexture->GetMediaPlayer() != MediaPlayer.Get()))
+			{
+				MediaTexture->SetMediaPlayer(MediaPlayer);
+				bApplyTextureUpdate = true;
+			}
+
+			if (bApplyTextureUpdate ||
+				(EnumHasAllFlags(Flags, UE::MediaPlateComponent::ESetUpTexturesFlags::ForceUpdateResource)))
+			{
+				MediaTexture->UpdateResource();
+			}
+		}
+	}
+
+	if (bApplyMaterialUpdate)
+	{
+		if (AMediaPlate* MediaPlate = GetOwner<AMediaPlate>())
+		{
+			if (UMaterialInterface* Material = MediaPlate->GetCurrentMaterial())
+			{
+				Material->RecacheUniformExpressions(false);
+			}
+		}
+	}
+}
 #if WITH_EDITOR
 
 void UMediaPlateComponent::SetVisibleMipsTilesCalculations(EMediaTextureVisibleMipsTiles InVisibleMipsTilesCalculations)
diff --git a/Engine/Plugins/Media/MediaPlate/Source/MediaPlate/Public/MediaPlate.h b/Engine/Plugins/Media/MediaPlate/Source/MediaPlate/Public/MediaPlate.h
--- a/Engine/Plugins/Media/MediaPlate/Source/MediaPlate/Public/MediaPlate.h
+++ b/Engine/Plugins/Media/MediaPlate/Source/MediaPlate/Public/MediaPlate.h
@@ -8,6 +8,7 @@
 #include "MediaPlate.generated.h"
 
 class UMaterialInterface;
+class UMaterialInstanceDynamic;
 class UMediaPlateComponent;
 
 /**
@@ -44,6 +45,11 @@ public:
 	 */
 	void ApplyCurrentMaterial();
 	void ApplyMaterial(UMaterialInterface* InMaterial);
+
+	/**
+	 * Sets up parameters (like the texture) that we use in the material.
+	 */
+	void SetMIDParameters(UMaterialInstanceDynamic* InMaterial);
 	
 	/** Get the last material assigned to the static mesh, at index 0. */
 	UMaterialInterface* GetLastMaterial() const { return LastMaterial; }
diff --git a/Engine/Plugins/Media/MediaPlate/Source/MediaPlate/Public/MediaPlateComponent.h b/Engine/Plugins/Media/MediaPlate/Source/MediaPlate/Public/MediaPlateComponent.h
--- a/Engine/Plugins/Media/MediaPlate/Source/MediaPlate/Public/MediaPlateComponent.h
+++ b/Engine/Plugins/Media/MediaPlate/Source/MediaPlate/Public/MediaPlateComponent.h
@@ -8,6 +8,7 @@
 #include "MediaPlayerProxyInterface.h"
 #include "MediaSource.h"
 #include "MediaTextureTracker.h"
+#include "Misc/EnumClassFlags.h"
 
 #include "MediaPlateComponent.generated.h"
 
@@ -19,6 +20,10 @@ class UMediaSoundComponent;
 class UMediaSource;
 class UMediaTexture;
 
+namespace UE::MediaPlateComponent
+{
+	enum class ESetUpTexturesFlags;
+}
 
 /**
  * This is a component for AMediaPlate that can play and show media in the world.
@@ -51,7 +56,7 @@ public:
 	 * Call this get our media texture.
 	 */
 	UFUNCTION(BlueprintCallable, Category = "Media|MediaPlateComponent")
-	UMediaTexture* GetMediaTexture();
+	UMediaTexture* GetMediaTexture(int32 Index = 0);
 
 	/**
 	 * Call this to open the media.
@@ -220,6 +225,11 @@ public:
 	 */
 	bool GetWantsToPlayWhenVisible() const { return bWantsToPlayWhenVisible; }
 
+	/**
+	 * Called from AMediaPlate to set how many media textures the material needs.
+	 */
+	void SetNumberOfTextures(int32 NumTextures);
+
 #if WITH_EDITOR
 	/**
 	 * Call this to get the mip tile calculations mesh mode.
@@ -247,6 +257,7 @@ public:
 	virtual bool ProxyIsPlaylistIndexPlaying(int32 Index) const override;
 	virtual void ProxyOpenPlaylistIndex(int32 Index) override;
 	virtual void ProxySetPlayOnOpen(bool bInPlayOnOpen) override;
+	virtual void ProxySetTextureBlend(int32 TextureIndex, float Blend) override;
 
 private:
 	/**
@@ -309,9 +320,9 @@ private:
 	/** Name for our playlist. */
 	static FLazyName MediaPlaylistName;
 
-	/** Holds the media player. */
+	/** Holds the media textures. */
 	UPROPERTY(Instanced)
-	TObjectPtr<UMediaTexture> MediaTexture;
+	TArray<TObjectPtr<UMediaTexture>> MediaTextures;
 
 	/** This component's media player */
 	UPROPERTY(Instanced)
@@ -395,4 +406,8 @@ private:
 	UFUNCTION()
 	void OnMediaEnd();
 
+	/**
+	 * Sets up the textures we have.
+	 */
+	void SetUpTextures(UE::MediaPlateComponent::ESetUpTexturesFlags Flags);
 };
diff --git a/Engine/Source/Runtime/MediaAssets/Public/MediaPlayerProxyInterface.h b/Engine/Source/Runtime/MediaAssets/Public/MediaPlayerProxyInterface.h
--- a/Engine/Source/Runtime/MediaAssets/Public/MediaPlayerProxyInterface.h
+++ b/Engine/Source/Runtime/MediaAssets/Public/MediaPlayerProxyInterface.h
@@ -78,6 +78,11 @@ public:
 	 */
 	virtual void ProxyOpenPlaylistIndex(int32 Index) = 0;
 
+	/**
+	 * Set the blend value for a texture that was retrieved from ProxyGetMediaTexture.
+	 */
+	virtual void ProxySetTextureBlend(int32 TextureIndex, float Blend) = 0;
+
 	/**
 	 * Set the player to play on open.
 	 */
