MediaPlate: Fixed crash when opening a media plate in the editor without a player.

diff --git a/Engine/Plugins/Media/MediaPlate/Source/MediaPlateEditor/Private/Toolkits/MediaPlateEditorToolkit.cpp b/Engine/Plugins/Media/MediaPlate/Source/MediaPlateEditor/Private/Toolkits/MediaPlateEditorToolkit.cpp
--- a/Engine/Plugins/Media/MediaPlate/Source/MediaPlateEditor/Private/Toolkits/MediaPlateEditorToolkit.cpp
+++ b/Engine/Plugins/Media/MediaPlate/Source/MediaPlateEditor/Private/Toolkits/MediaPlateEditorToolkit.cpp
@@ -248,7 +248,10 @@ void FMediaPlateEditorToolkit::BindCommands()
 	ToolkitCommands->MapAction(
 		Commands.CloseMedia,
 		FExecuteAction::CreateLambda([this] { MediaPlate->GetMediaPlayer()->Close(); }),
-		FCanExecuteAction::CreateLambda([this] { return !MediaPlate->GetMediaPlayer()->GetUrl().IsEmpty(); })
+		FCanExecuteAction::CreateLambda([this] {
+			TObjectPtr<UMediaPlayer> MediaPlayer = MediaPlate->GetMediaPlayer();
+			return (MediaPlayer != nullptr) && !MediaPlayer->GetUrl().IsEmpty();
+		})
 	);
 
 	ToolkitCommands->MapAction(
@@ -256,7 +259,7 @@ void FMediaPlateEditorToolkit::BindCommands()
 		FExecuteAction::CreateLambda([this]{ MediaPlate->GetMediaPlayer()->SetRate(GetForwardRate()); }),
 		FCanExecuteAction::CreateLambda([this]{
 			TObjectPtr<UMediaPlayer> MediaPlayer = MediaPlate->GetMediaPlayer();
-			return MediaPlayer->IsReady() && MediaPlayer->SupportsRate(GetForwardRate(), false);
+			return (MediaPlayer != nullptr) && MediaPlayer->IsReady() && MediaPlayer->SupportsRate(GetForwardRate(), false);
 		})
 	);
 
@@ -280,7 +283,7 @@ void FMediaPlateEditorToolkit::BindCommands()
 		FExecuteAction::CreateLambda([this]{ MediaPlate->Pause(); }),
 		FCanExecuteAction::CreateLambda([this]{
 			TObjectPtr<UMediaPlayer> MediaPlayer = MediaPlate->GetMediaPlayer();
-			return MediaPlayer->CanPause() && !MediaPlayer->IsPaused();
+			return (MediaPlayer != nullptr) &&MediaPlayer->CanPause() && !MediaPlayer->IsPaused();
 		})
 	);
 
@@ -289,7 +292,7 @@ void FMediaPlateEditorToolkit::BindCommands()
 		FExecuteAction::CreateLambda([this]{ MediaPlate->GetMediaPlayer()->Play(); }),
 		FCanExecuteAction::CreateLambda([this]{
 			TObjectPtr<UMediaPlayer> MediaPlayer = MediaPlate->GetMediaPlayer(); 
-			return MediaPlayer->IsReady() && (!MediaPlayer->IsPlaying() || (MediaPlayer->GetRate() != 1.0f));
+			return (MediaPlayer != nullptr) && MediaPlayer->IsReady() && (!MediaPlayer->IsPlaying() || (MediaPlayer->GetRate() != 1.0f));
 		})
 	);
 
@@ -307,7 +310,7 @@ void FMediaPlateEditorToolkit::BindCommands()
 		FExecuteAction::CreateLambda([this]{ MediaPlate->GetMediaPlayer()->SetRate(GetReverseRate()); } ),
 		FCanExecuteAction::CreateLambda([this]{
 			TObjectPtr<UMediaPlayer> MediaPlayer = MediaPlate->GetMediaPlayer();
-			return MediaPlayer->IsReady() && MediaPlayer->SupportsRate(GetReverseRate(), false);
+			return (MediaPlayer != nullptr) && MediaPlayer->IsReady() && MediaPlayer->SupportsRate(GetReverseRate(), false);
 		})
 	);
 
@@ -316,7 +319,7 @@ void FMediaPlateEditorToolkit::BindCommands()
 		FExecuteAction::CreateLambda([this]{ MediaPlate->GetMediaPlayer()->Rewind(); }),
 		FCanExecuteAction::CreateLambda([this]{
 			TObjectPtr<UMediaPlayer> MediaPlayer = MediaPlate->GetMediaPlayer(); 
-			return MediaPlate->GetMediaPlayer()->IsReady() && MediaPlayer->SupportsSeeking() && MediaPlayer->GetTime() > FTimespan::Zero();
+			return (MediaPlayer != nullptr) && MediaPlate->GetMediaPlayer()->IsReady() && MediaPlayer->SupportsSeeking() && MediaPlayer->GetTime() > FTimespan::Zero();
 		})
 	);
 }
@@ -411,7 +414,10 @@ TSharedRef<SDockTab> FMediaPlateEditorToolkit::HandleTabManagerSpawnTab(const FS
 	}
 	else if (TabIdentifier == MediaPlateEditorToolkit::ViewerTabId)
 	{
-		TabWidget = SNew(SMediaPlayerEditorViewer, *MediaPlayer, MediaTexture, Style, false);
+		if (MediaPlayer != nullptr)
+		{
+			TabWidget = SNew(SMediaPlayerEditorViewer, *MediaPlayer, MediaTexture, Style, false);
+		}
 	}
 
 	return SNew(SDockTab)
