MediaCompositing: Fixed media plate tracks not blending appropriately.

diff --git a/Engine/Plugins/Media/MediaCompositing/Source/MediaCompositing/Private/MovieScene/MovieSceneMediaData.cpp b/Engine/Plugins/Media/MediaCompositing/Source/MediaCompositing/Private/MovieScene/MovieSceneMediaData.cpp
--- a/Engine/Plugins/Media/MediaCompositing/Source/MediaCompositing/Private/MovieScene/MovieSceneMediaData.cpp
+++ b/Engine/Plugins/Media/MediaCompositing/Source/MediaCompositing/Private/MovieScene/MovieSceneMediaData.cpp
@@ -47,7 +47,7 @@ void FMovieSceneMediaData::SeekOnOpen(FTimespan Time)
 }
 
 
-void FMovieSceneMediaData::Setup(UMediaPlayer* OverrideMediaPlayer, UObject* InPlayerProxy, int32 InProxyTextureIndex)
+void FMovieSceneMediaData::Setup(UMediaPlayer* OverrideMediaPlayer, UObject* InPlayerProxy, int32 InProxyLayerIndex, int32 InProxyTextureIndex)
 {
 	// Ensure we don't already have a media player set. Setup should only be called once
 	check(!MediaPlayer);
@@ -66,17 +66,13 @@ void FMovieSceneMediaData::Setup(UMediaPlayer* OverrideMediaPlayer, UObject* InP
 	MediaPlayer->OnMediaEvent().AddRaw(this, &FMovieSceneMediaData::HandleMediaPlayerEvent);
 	MediaPlayer->AddToRoot();
 	ProxyMediaTexture.Reset();
+	ProxyLayerIndex = InProxyLayerIndex;
 	ProxyTextureIndex = InProxyTextureIndex;
 
 	// Do we have a valid proxy object?
 	if ((InPlayerProxy != nullptr) && (InPlayerProxy->Implements<UMediaPlayerProxyInterface>()))
 	{
 		PlayerProxy = InPlayerProxy;
-		IMediaPlayerProxyInterface* PlayerProxyInterface = Cast<IMediaPlayerProxyInterface>(PlayerProxy);
-		if (PlayerProxyInterface != nullptr)
-		{
-			ProxyMediaTexture = PlayerProxyInterface->ProxyGetMediaTexture(ProxyTextureIndex);
-		}
 	}
 	else
 	{
@@ -105,6 +101,14 @@ void FMovieSceneMediaData::StartUsingProxyMediaTexture()
 {
 	if (PlayerProxy != nullptr)
 	{
+		if (ProxyMediaTexture == nullptr)
+		{
+			IMediaPlayerProxyInterface* PlayerProxyInterface = Cast<IMediaPlayerProxyInterface>(PlayerProxy);
+			if (PlayerProxyInterface != nullptr)
+			{
+				ProxyMediaTexture = PlayerProxyInterface->ProxyGetMediaTexture(ProxyLayerIndex, ProxyTextureIndex);
+			}
+		}
 		if (ProxyMediaTexture != nullptr)
 		{
 			ProxyMediaTexture->SetMediaPlayer(MediaPlayer);
@@ -122,6 +126,12 @@ void FMovieSceneMediaData::StopUsingProxyMediaTexture()
 			{
 				ProxyMediaTexture->SetMediaPlayer(nullptr);
 			}
+			IMediaPlayerProxyInterface* PlayerProxyInterface = Cast<IMediaPlayerProxyInterface>(PlayerProxy);
+			if (PlayerProxyInterface != nullptr)
+			{
+				PlayerProxyInterface->ProxyReleaseMediaTexture(ProxyLayerIndex, ProxyTextureIndex);
+			}
+			ProxyMediaTexture = nullptr;
 		}
 	}
 }
diff --git a/Engine/Plugins/Media/MediaCompositing/Source/MediaCompositing/Private/MovieScene/MovieSceneMediaSection.cpp b/Engine/Plugins/Media/MediaCompositing/Source/MediaCompositing/Private/MovieScene/MovieSceneMediaSection.cpp
--- a/Engine/Plugins/Media/MediaCompositing/Source/MediaCompositing/Private/MovieScene/MovieSceneMediaSection.cpp
+++ b/Engine/Plugins/Media/MediaCompositing/Source/MediaCompositing/Private/MovieScene/MovieSceneMediaSection.cpp
@@ -143,7 +143,7 @@ UMediaSource* UMovieSceneMediaSection::GetMediaSourceOrProxy(IMovieScenePlayer&
 					IMediaPlayerProxyInterface* Proxy = Cast<IMediaPlayerProxyInterface>(Component);
 					if (Proxy != nullptr)
 					{
-						OutMediaSource = Proxy->GetMediaSourceFromIndex(InMediaSourceProxyIndex);
+						OutMediaSource = Proxy->ProxyGetMediaSourceFromIndex(InMediaSourceProxyIndex);
 						break;
 					}
 				}
diff --git a/Engine/Plugins/Media/MediaCompositing/Source/MediaCompositing/Private/MovieScene/MovieSceneMediaTemplate.cpp b/Engine/Plugins/Media/MediaCompositing/Source/MediaCompositing/Private/MovieScene/MovieSceneMediaTemplate.cpp
--- a/Engine/Plugins/Media/MediaCompositing/Source/MediaCompositing/Private/MovieScene/MovieSceneMediaTemplate.cpp
+++ b/Engine/Plugins/Media/MediaCompositing/Source/MediaCompositing/Private/MovieScene/MovieSceneMediaTemplate.cpp
@@ -142,7 +142,7 @@ struct FMediaSectionExecutionToken
 				(PlayerProxy);
 			if (PlayerProxyInterface != nullptr)
 			{
-				PlayerProxyInterface->ProxySetTextureBlend(SectionData.GetProxyTextureIndex(), ProxyTextureBlend);
+				PlayerProxyInterface->ProxySetTextureBlend(SectionData.GetProxyLayerIndex(), SectionData.GetProxyTextureIndex(), ProxyTextureBlend);
 				// Can we control the player?
 				if (PlayerProxyInterface->IsExternalControlAllowed() == false)
 				{
@@ -367,8 +367,14 @@ void FMovieSceneMediaSectionTemplate::Initialize(const FMovieSceneEvaluationOper
 	FMovieSceneMediaData* SectionData = PersistentData.FindSectionData<FMovieSceneMediaData>();
 	if (SectionData == nullptr)
 	{
-		int32 ProxyTextureIndex = (MediaSection != nullptr) ? MediaSection->TextureIndex : 0;
-		
+		int32 ProxyTextureIndex = 0;
+		int32 ProxyLayerIndex = 0;
+		if (MediaSection != nullptr)
+		{
+			ProxyTextureIndex = MediaSection->TextureIndex;
+			ProxyLayerIndex = MediaSection->GetRowIndex();
+		}
+
 		// Are we overriding the media player?
 		UMediaPlayer* MediaPlayer = Params.MediaPlayer;
 		UObject* PlayerProxy = nullptr;
@@ -387,18 +393,6 @@ void FMovieSceneMediaSectionTemplate::Initialize(const FMovieSceneEvaluationOper
 						if (BoundObject != nullptr)
 						{
 							MediaAssetsModule->GetPlayerFromObject(BoundObject, PlayerProxy);
-							if (PlayerProxy != nullptr)
-							{
-								IMediaPlayerProxyInterface* PlayerProxyInterface = Cast<IMediaPlayerProxyInterface>(PlayerProxy);
-								if (PlayerProxyInterface != nullptr)
-								{
-									int32 NumberOfTextures = PlayerProxyInterface->ProxyGetNumberOfMediaTextures();
-									if (NumberOfTextures > 0)
-									{
-										ProxyTextureIndex %= NumberOfTextures;
-									}
-								}
-							}
 							break;
 						}
 					}
@@ -408,7 +402,7 @@ void FMovieSceneMediaSectionTemplate::Initialize(const FMovieSceneEvaluationOper
 
 		// Add section data.
 		SectionData = &PersistentData.AddSectionData<FMovieSceneMediaData>();
-		SectionData->Setup(MediaPlayer, PlayerProxy, ProxyTextureIndex);
+		SectionData->Setup(MediaPlayer, PlayerProxy, ProxyLayerIndex, ProxyTextureIndex);
 	}
 
 	if (!ensure(SectionData != nullptr))
@@ -508,7 +502,7 @@ void FMovieSceneMediaSectionTemplate::TearDown(FPersistentEvaluationData& Persis
 		IMediaPlayerProxyInterface* PlayerProxyInterface = Cast<IMediaPlayerProxyInterface>(PlayerProxy);
 		if (PlayerProxyInterface != nullptr)
 		{
-			PlayerProxyInterface->ProxySetTextureBlend(SectionData->GetProxyTextureIndex(), 0.0f);
+			PlayerProxyInterface->ProxySetTextureBlend(SectionData->GetProxyLayerIndex(), SectionData->GetProxyTextureIndex(), 0.0f);
 		}
 	}
 
diff --git a/Engine/Plugins/Media/MediaCompositing/Source/MediaCompositing/Public/MovieSceneMediaData.h b/Engine/Plugins/Media/MediaCompositing/Source/MediaCompositing/Public/MovieSceneMediaData.h
--- a/Engine/Plugins/Media/MediaCompositing/Source/MediaCompositing/Public/MovieSceneMediaData.h
+++ b/Engine/Plugins/Media/MediaCompositing/Source/MediaCompositing/Public/MovieSceneMediaData.h
@@ -38,6 +38,11 @@ public:
 	 */
 	UObject* GetPlayerProxy() { return PlayerProxy.Get(); }
 
+	/**
+	 * Get the layer index we are using (when using a proxy).
+	 */
+	int32 GetProxyLayerIndex() { return ProxyLayerIndex; }
+
 	/**
 	 * Get the texture index we are using (when using a proxy).
 	 */
@@ -51,7 +56,7 @@ public:
 	void SeekOnOpen(FTimespan Time);
 
 	/** Set up this persistent data object. */
-	void Setup(UMediaPlayer* OverrideMediaPlayer, UObject* InPlayerProxy, int32 InProxyTextureIndex);
+	void Setup(UMediaPlayer* OverrideMediaPlayer, UObject* InPlayerProxy, int32 InProxyLayerIndex, int32 InProxyTextureIndex);
 
 	/**
 	 * Called from FMovieSceneMediaSectionTemplate::Initialize.
@@ -86,6 +91,8 @@ private:
 	TWeakObjectPtr<UObject> PlayerProxy;
 	/** Media texture allocated from the proxy. */
 	TWeakObjectPtr<UMediaTexture> ProxyMediaTexture;
+	/** Layer that this section should reside in. */
+	int32 ProxyLayerIndex;
 	/** Index of texture allocated from the proxy. */
 	int32 ProxyTextureIndex;
 
diff --git a/Engine/Plugins/Media/MediaPlate/Source/MediaPlate/Private/MediaPlateComponent.cpp b/Engine/Plugins/Media/MediaPlate/Source/MediaPlate/Private/MediaPlateComponent.cpp
--- a/Engine/Plugins/Media/MediaPlate/Source/MediaPlate/Private/MediaPlateComponent.cpp
+++ b/Engine/Plugins/Media/MediaPlate/Source/MediaPlate/Private/MediaPlateComponent.cpp
@@ -202,21 +202,18 @@ UMediaTexture* UMediaPlateComponent::GetMediaTexture(int32 Index)
 	}
 	else
 	{
-		UE_LOG(LogMediaPlate, Warning, TEXT("Texture index %d not available."), Index);
+		UE_CALL_ONCE([Index](){
+			UE_LOG(LogMediaPlate, Warning, TEXT("Material does not support texture index %d. Either remove the number of cross fades or change the material."), Index);
+		});
 	}
 	return MediaTexture;
 }
 
 void UMediaPlateComponent::Open()
-{
-	ProxyOpenPlaylistIndex(0);
-}
-
-void UMediaPlateComponent::ProxyOpenPlaylistIndex(int32 Index)
 {
 	bIsMediaPlatePlaying = true;
 	CurrentRate = bPlayOnOpen ? 1.0f : 0.0f;
-	PlaylistIndex = Index;
+	PlaylistIndex = 0;
 
 	if (IsVisible())
 	{
@@ -226,7 +223,7 @@ void UMediaPlateComponent::ProxyOpenPlaylistIndex(int32 Index)
 			UMediaSource* MediaSource = nullptr;
 			if (MediaPlaylist != nullptr)
 			{
-				MediaSource = MediaPlaylist->Get(Index);
+				MediaSource = MediaPlaylist->Get(PlaylistIndex);
 			}
 			bIsPlaying = PlayMediaSource(MediaSource, bPlayOnOpen);
 		}
@@ -617,7 +614,7 @@ const FMediaSourceCacheSettings& UMediaPlateComponent::GetCacheSettings() const
 	return CacheSettings;
 }
 
-UMediaSource* UMediaPlateComponent::GetMediaSourceFromIndex(int32 Index) const
+UMediaSource* UMediaPlateComponent::ProxyGetMediaSourceFromIndex(int32 Index) const
 {
 	UMediaSource* MediaSource = nullptr;
 	if (MediaPlaylist != nullptr)
@@ -627,32 +624,53 @@ UMediaSource* UMediaPlateComponent::GetMediaSourceFromIndex(int32 Index) const
 	return MediaSource;
 }
 
-void UMediaPlateComponent::ProxyClose()
+UMediaTexture* UMediaPlateComponent::ProxyGetMediaTexture(int32 LayerIndex, int32 TextureIndex)
 {
-	Close();
-}
+	UMediaTexture* MediaTexture = GetMediaTexture(TextureIndex);
 
-UMediaTexture* UMediaPlateComponent::ProxyGetMediaTexture(int32 TextureIndex)
-{
-	return GetMediaTexture(TextureIndex);
-}
+	if (TextureLayers.Num() < LayerIndex + 1)
+	{
+		TextureLayers.SetNum(LayerIndex + 1);
+	}
 
-int32 UMediaPlateComponent::ProxyGetNumberOfMediaTextures() const
-{
-	return MediaTextures.Num();
-}
+	// Fill up an empty slot if there is one.
+	bool bIsTextureSet = false;
+	for (int32 Index = 0; Index < TextureLayers[LayerIndex].Num(); ++Index)
+	{
+		if (TextureLayers[LayerIndex][Index] < 0)
+		{
+			TextureLayers[LayerIndex][Index] = TextureIndex;
+			bIsTextureSet = true;
+			break;
+		}
+	}
+	if (bIsTextureSet == false)
+	{
+		TextureLayers[LayerIndex].Add(TextureIndex);
+	}
 
-bool UMediaPlateComponent::ProxyIsPlaylistIndexPlaying(int32 Index) const
-{
-	return bIsMediaPlatePlaying && (Index == PlaylistIndex);
+	UpdateTextureLayers();
+
+	return MediaTexture;
 }
 
-void UMediaPlateComponent::ProxySetPlayOnOpen(bool bInPlayOnOpen)
+void UMediaPlateComponent::ProxyReleaseMediaTexture(int32 LayerIndex, int32 TextureIndex)
 {
-	bPlayOnOpen = bInPlayOnOpen;
+	ProxySetTextureBlend(LayerIndex, TextureIndex, 0.0f);
+
+	for (int32 Index = 0; Index < TextureLayers[LayerIndex].Num(); ++Index)
+	{
+		if (TextureLayers[LayerIndex][Index] == TextureIndex)
+		{
+			TextureLayers[LayerIndex][Index] = -1;
+			break;
+		}
+	}
+	
+	UpdateTextureLayers();
 }
 
-void UMediaPlateComponent::ProxySetTextureBlend(int32 TextureIndex, float Blend)
+void UMediaPlateComponent::ProxySetTextureBlend(int32 LayerIndex, int32 TextureIndex, float Blend)
 {
 	if (AMediaPlate* MediaPlate = GetOwner<AMediaPlate>())
 	{
@@ -661,10 +679,25 @@ void UMediaPlateComponent::ProxySetTextureBlend(int32 TextureIndex, float Blend)
 			UMaterialInstanceDynamic* MID = Cast<UMaterialInstanceDynamic>(Material);
 			if (MID != nullptr)
 			{
-				static const FString BaseBlendName = TEXT("Blend");
-				FString BlendName = BaseBlendName;
-				BlendName.AppendInt(TextureIndex);
-				MID->SetScalarParameterValue(FName(*BlendName), Blend);
+				int32 MatNumLayers = MediaTextures.Num() / MatNumTexPerLayer;
+				if (LayerIndex < MatNumLayers)
+				{
+					const TArray<int32>& Layer = TextureLayers[LayerIndex];
+					for (int32 LayerTexIndex = 0;
+						(LayerTexIndex < MatNumTexPerLayer) && (LayerTexIndex < Layer.Num());
+						LayerTexIndex++)
+					{
+						if (Layer[LayerTexIndex] == TextureIndex)
+						{
+							int32 MatTexIndex = LayerIndex * MatNumTexPerLayer + LayerTexIndex;
+							static const FString BaseBlendName = TEXT("Blend");
+							FString BlendName = BaseBlendName;
+							BlendName.AppendInt(MatTexIndex);
+							MID->SetScalarParameterValue(FName(*BlendName), Blend);
+							break;
+						}
+					}
+				}
 			}
 		}
 	}
@@ -925,6 +958,47 @@ void UMediaPlateComponent::SetUpTextures(UE::MediaPlateComponent::ESetUpTextures
 		}
 	}
 }
+
+void UMediaPlateComponent::UpdateTextureLayers()
+{
+	if (AMediaPlate* MediaPlate = GetOwner<AMediaPlate>())
+	{
+		if (UMaterialInterface* Material = MediaPlate->GetCurrentMaterial())
+		{
+			UMaterialInstanceDynamic* MID = Cast<UMaterialInstanceDynamic>(Material);
+			if (MID != nullptr)
+			{
+				// Go through each layer.
+				int32 MatNumLayers = MediaTextures.Num() / MatNumTexPerLayer;
+				static const FString BaseTextureName = TEXT("MediaTexture");
+				int32 NumLayers = FMath::Min(MatNumLayers, TextureLayers.Num());
+				for (int32 LayerIndex = 0; LayerIndex < NumLayers; ++LayerIndex)
+				{
+					// Go through each texture in the layer.
+					const TArray<int32>& Layer = TextureLayers[LayerIndex];
+					int32 NumTex = FMath::Min(MatNumTexPerLayer, Layer.Num());
+					for (int32 LayerTexIndex = 0; LayerTexIndex < NumTex; LayerTexIndex++)
+					{
+						// Set the texture in the material according to the layer data.
+						int32 TextureIndex = Layer[LayerTexIndex];
+						if (TextureIndex >= 0)
+						{
+							int32 MatTexIndex = LayerIndex * MatNumTexPerLayer + LayerTexIndex;
+							FString TextureName = BaseTextureName;
+							if (MatTexIndex != 0)
+							{
+								TextureName.AppendInt(MatTexIndex);
+							}
+							MID->SetTextureParameterValue(FName(*TextureName), MediaTextures[TextureIndex]);
+						}
+					}
+				}
+			}
+		}
+	}
+}
+
+
 #if WITH_EDITOR
 
 void UMediaPlateComponent::SetVisibleMipsTilesCalculations(EMediaTextureVisibleMipsTiles InVisibleMipsTilesCalculations)
diff --git a/Engine/Plugins/Media/MediaPlate/Source/MediaPlate/Public/MediaPlateComponent.h b/Engine/Plugins/Media/MediaPlate/Source/MediaPlate/Public/MediaPlateComponent.h
--- a/Engine/Plugins/Media/MediaPlate/Source/MediaPlate/Public/MediaPlateComponent.h
+++ b/Engine/Plugins/Media/MediaPlate/Source/MediaPlate/Public/MediaPlateComponent.h
@@ -251,14 +251,10 @@ public:
 	virtual bool SetProxyRate(float Rate) override;
 	virtual bool IsExternalControlAllowed() override;
 	virtual const FMediaSourceCacheSettings& GetCacheSettings() const override;
-	virtual UMediaSource* GetMediaSourceFromIndex(int32 Index) const override;
-	virtual void ProxyClose() override;
-	virtual UMediaTexture* ProxyGetMediaTexture(int32 TextureIndex) override;
-	virtual int32 ProxyGetNumberOfMediaTextures() const override;
-	virtual bool ProxyIsPlaylistIndexPlaying(int32 Index) const override;
-	virtual void ProxyOpenPlaylistIndex(int32 Index) override;
-	virtual void ProxySetPlayOnOpen(bool bInPlayOnOpen) override;
-	virtual void ProxySetTextureBlend(int32 TextureIndex, float Blend) override;
+	virtual UMediaSource* ProxyGetMediaSourceFromIndex(int32 Index) const override;
+	virtual UMediaTexture* ProxyGetMediaTexture(int32 LayerIndex, int32 TextureIndex) override;
+	virtual void ProxyReleaseMediaTexture(int32 LayerIndex, int32 TextureIndex) override;
+	virtual void ProxySetTextureBlend(int32 LayerIndex, int32 TextureIndex, float Blend) override;
 
 private:
 	/**
@@ -321,6 +317,9 @@ private:
 	UPROPERTY()
 	float LetterboxAspectRatio = 0.0f;
 
+	/** Number of textures we have per layer in the material. */
+	const int32 MatNumTexPerLayer = 2;
+
 	UPROPERTY()
 	FVector2D MeshRange = FVector2D(360.0f, 180.0f);
 
@@ -348,6 +347,14 @@ private:
 	/** True if we should resume where we left off when we open the media. */
 	bool bResumeWhenOpened = false;
 
+	/**
+	 * Contains all of our layers.
+	 * Each layer contains which textures it has.
+	 * int32 is an index into MediaTextures.
+	 * -1 signifies no entry.
+	 */
+	TArray<TArray<int32>> TextureLayers;
+
 	/**
 	 * Plays a media source.
 	 * 
@@ -419,4 +426,9 @@ private:
 	 * Sets up the textures we have.
 	 */
 	void SetUpTextures(UE::MediaPlateComponent::ESetUpTexturesFlags Flags);
+
+	/**
+	 * Sets textures in our material according to the layer assignments.
+	 */
+	void UpdateTextureLayers();
 };
diff --git a/Engine/Source/Runtime/MediaAssets/Public/MediaPlayerProxyInterface.h b/Engine/Source/Runtime/MediaAssets/Public/MediaPlayerProxyInterface.h
--- a/Engine/Source/Runtime/MediaAssets/Public/MediaPlayerProxyInterface.h
+++ b/Engine/Source/Runtime/MediaAssets/Public/MediaPlayerProxyInterface.h
@@ -56,41 +56,20 @@ public:
 	/**
 	 * Get the media source for a given index.
 	 */
-	virtual UMediaSource* GetMediaSourceFromIndex(int32 Index) const = 0;
-
-	/**
-	 * Close the player.
-	 */
-	virtual void ProxyClose() = 0;
+	virtual UMediaSource* ProxyGetMediaSourceFromIndex(int32 Index) const = 0;
 
 	/**
 	 * Get a media texture that we can assign a media player to.
 	 */
-	virtual UMediaTexture* ProxyGetMediaTexture(int32 TextureIndex) = 0;
-
-	/**
-	 * Gets the number of media textures we can get with ProxyGetMediaTexture.
-	 */
-	virtual int32 ProxyGetNumberOfMediaTextures() const = 0;
+	virtual UMediaTexture* ProxyGetMediaTexture(int32 LayerIndex, int32 TextureIndex) = 0;
 
 	/**
-	 * Ask if a specific track in the playlist is playing.
+	 * Release a media texture that was retrieved from ProxyGetMediaTexture.
 	 */
-	virtual bool ProxyIsPlaylistIndexPlaying(int32 Index) const = 0;
-
-	/**
-	 * Open a specific track in the playlist.
-	 */
-	virtual void ProxyOpenPlaylistIndex(int32 Index) = 0;
+	virtual void ProxyReleaseMediaTexture(int32 LayerIndex, int32 TextureIndex) = 0;
 
 	/**
 	 * Set the blend value for a texture that was retrieved from ProxyGetMediaTexture.
 	 */
-	virtual void ProxySetTextureBlend(int32 TextureIndex, float Blend) = 0;
-
-	/**
-	 * Set the player to play on open.
-	 */
-	virtual void ProxySetPlayOnOpen(bool bInPlayOnOpen) = 0;
-	
+	virtual void ProxySetTextureBlend(int32 LayerIndex, int32 TextureIndex, float Blend) = 0;
 };
