Asset Definition: Convert ImagePlateFileSequence and MediaPlate

diff --git a/Engine/Plugins/Experimental/ImagePlate/Source/ImagePlateEditor/ImagePlateEditor.Build.cs b/Engine/Plugins/Experimental/ImagePlate/Source/ImagePlateEditor/ImagePlateEditor.Build.cs
--- a/Engine/Plugins/Experimental/ImagePlate/Source/ImagePlateEditor/ImagePlateEditor.Build.cs
+++ b/Engine/Plugins/Experimental/ImagePlate/Source/ImagePlateEditor/ImagePlateEditor.Build.cs
@@ -9,7 +9,7 @@ namespace UnrealBuildTool.Rules
 			PrivateDependencyModuleNames.AddRange(
 				new string[]
 				{
-					"AssetTools",
+					"AssetDefinition",
 					"Core",
 					"CoreUObject",
 					"EditorFramework",
diff --git a/Engine/Plugins/Experimental/ImagePlate/Source/ImagePlateEditor/Private/AssetDefinition_ImagePlateFileSequence.h b/Engine/Plugins/Experimental/ImagePlate/Source/ImagePlateEditor/Private/AssetDefinition_ImagePlateFileSequence.h
new file mode 100644
--- /dev/null
+++ b/Engine/Plugins/Experimental/ImagePlate/Source/ImagePlateEditor/Private/AssetDefinition_ImagePlateFileSequence.h
@@ -0,0 +1,26 @@
+// Copyright Epic Games, Inc. All Rights Reserved.
+
+#pragma once
+
+#include "ImagePlateFileSequence.h"
+#include "AssetDefinitionDefault.h"
+
+#include "AssetDefinition_ImagePlateFileSequence.generated.h"
+
+UCLASS()
+class UAssetDefinition_ImagePlateFileSequence : public UAssetDefinitionDefault
+{
+	GENERATED_BODY()
+
+public:
+	// UAssetDefinition Begin
+	virtual FText GetAssetDisplayName() const override { return NSLOCTEXT("AssetTypeActions", "AssetTypeActions_ImagePlateFileSequence", "Image Plate File Sequence"); }
+	virtual FLinearColor GetAssetColor() const override { return FLinearColor(FColor(105, 0, 60)); }
+	virtual TSoftClassPtr<UObject> GetAssetClass() const override { return UImagePlateFileSequence::StaticClass(); }
+	virtual TConstArrayView<FAssetCategoryPath> GetAssetCategories() const override
+	{
+		static const auto Categories = { EAssetCategoryPaths::Media };
+		return Categories;
+	}
+	// UAssetDefinition End
+};
diff --git a/Engine/Plugins/Experimental/ImagePlate/Source/ImagePlateEditor/Private/AssetTypeActions_ImagePlateFileSequence.h b/Engine/Plugins/Experimental/ImagePlate/Source/ImagePlateEditor/Private/AssetTypeActions_ImagePlateFileSequence.h
deleted file mode 100644
diff --git a/Engine/Plugins/Experimental/ImagePlate/Source/ImagePlateEditor/Private/ImagePlateEditor.cpp b/Engine/Plugins/Experimental/ImagePlate/Source/ImagePlateEditor/Private/ImagePlateEditor.cpp
--- a/Engine/Plugins/Experimental/ImagePlate/Source/ImagePlateEditor/Private/ImagePlateEditor.cpp
+++ b/Engine/Plugins/Experimental/ImagePlate/Source/ImagePlateEditor/Private/ImagePlateEditor.cpp
@@ -1,10 +1,10 @@
 // Copyright Epic Games, Inc. All Rights Reserved.
 
+#include "ImagePlateFileSequence.h"
 #include "Brushes/SlateImageBrush.h"
 #include "Styling/SlateStyle.h"
 #include "Misc/Paths.h"
 #include "Styling/SlateStyleRegistry.h"
-#include "AssetTypeActions_ImagePlateFileSequence.h"
 #include "ISettingsModule.h"
 
 #define LOCTEXT_NAMESPACE "ImagePlateEditor"
@@ -62,10 +62,6 @@ public:
 	virtual void StartupModule() override
 	{
 		FImagePlateEditorStyle::Get();
-		AssetTypeActions = MakeShared<FAssetTypeActions_ImagePlateFileSequence>();
-
-		IAssetTools& AssetToolsModule = FModuleManager::Get().LoadModuleChecked<FAssetToolsModule>("AssetTools").Get();
-		AssetToolsModule.RegisterAssetTypeActions(AssetTypeActions.ToSharedRef());
 
 		ISettingsModule& SettingsModule = FModuleManager::LoadModuleChecked<ISettingsModule>("Settings");
 		SettingsModule.RegisterSettings("Project", "Plugins", "ImagePlate",
@@ -78,22 +74,13 @@ public:
 	virtual void ShutdownModule() override
 	{
 		FImagePlateEditorStyle::Destroy();
-
-		FAssetToolsModule* AssetToolsModule = FModuleManager::Get().GetModulePtr<FAssetToolsModule>("AssetTools");
-		if (AssetToolsModule)
-		{
-			AssetToolsModule->Get().UnregisterAssetTypeActions(AssetTypeActions.ToSharedRef());
-		}
 		
 		ISettingsModule* SettingsModule = FModuleManager::GetModulePtr<ISettingsModule>("Settings");
 		if (SettingsModule)
 		{
 			SettingsModule->UnregisterSettings("Project", "Plugins", "ImagePlate");
 		}
-
 	}
-
-	TSharedPtr<FAssetTypeActions_ImagePlateFileSequence> AssetTypeActions;
 };
 
 IMPLEMENT_MODULE(FImagePlateEditorModule, ImagePlateEditor);
diff --git a/Engine/Plugins/Media/MediaPlate/Source/MediaPlateEditor/MediaPlateEditor.Build.cs b/Engine/Plugins/Media/MediaPlate/Source/MediaPlateEditor/MediaPlateEditor.Build.cs
--- a/Engine/Plugins/Media/MediaPlate/Source/MediaPlateEditor/MediaPlateEditor.Build.cs
+++ b/Engine/Plugins/Media/MediaPlate/Source/MediaPlateEditor/MediaPlateEditor.Build.cs
@@ -18,6 +18,7 @@ public class MediaPlateEditor : ModuleRules
 		PrivateDependencyModuleNames.AddRange(
 			new string[]
 			{
+				"AssetDefinition",
 				"CinematicCamera",
 				"CoreUObject",
 				"DesktopWidgets",
diff --git a/Engine/Plugins/Media/MediaPlate/Source/MediaPlateEditor/Private/AssetTools/AssetDefinition_MediaPlate.cpp b/Engine/Plugins/Media/MediaPlate/Source/MediaPlateEditor/Private/AssetTools/AssetDefinition_MediaPlate.cpp
new file mode 100644
--- /dev/null
+++ b/Engine/Plugins/Media/MediaPlate/Source/MediaPlateEditor/Private/AssetTools/AssetDefinition_MediaPlate.cpp
@@ -0,0 +1,21 @@
+// Copyright Epic Games, Inc. All Rights Reserved.
+
+#include "AssetDefinition_MediaPlate.h"
+#include "MediaPlateComponent.h"
+#include "MediaPlateEditorStyle.h"
+#include "Toolkits/MediaPlateEditorToolkit.h"
+
+#define LOCTEXT_NAMESPACE "AssetTypeActions"
+
+EAssetCommandResult UAssetDefinition_MediaPlate::OpenAssets(const FAssetOpenArgs& OpenArgs) const
+{
+	for (UMediaPlateComponent* MediaPlate : OpenArgs.LoadObjects<UMediaPlateComponent>())
+	{
+		TSharedRef<FMediaPlateEditorToolkit> EditorToolkit = MakeShareable(new FMediaPlateEditorToolkit(FMediaPlateEditorStyle::Get()));
+		EditorToolkit->Initialize(MediaPlate, OpenArgs.GetToolkitMode(), OpenArgs.ToolkitHost);
+	}
+
+	return EAssetCommandResult::Handled;
+}
+
+#undef LOCTEXT_NAMESPACE
diff --git a/Engine/Plugins/Media/MediaPlate/Source/MediaPlateEditor/Private/AssetTools/AssetDefinition_MediaPlate.h b/Engine/Plugins/Media/MediaPlate/Source/MediaPlateEditor/Private/AssetTools/AssetDefinition_MediaPlate.h
new file mode 100644
--- /dev/null
+++ b/Engine/Plugins/Media/MediaPlate/Source/MediaPlateEditor/Private/AssetTools/AssetDefinition_MediaPlate.h
@@ -0,0 +1,27 @@
+// Copyright Epic Games, Inc. All Rights Reserved.
+
+#pragma once
+
+#include "MediaPlateComponent.h"
+#include "AssetDefinitionDefault.h"
+
+#include "AssetDefinition_MediaPlate.generated.h"
+
+UCLASS()
+class UAssetDefinition_MediaPlate : public UAssetDefinitionDefault
+{
+	GENERATED_BODY()
+
+public:
+	// UAssetDefinition Begin
+	virtual FText GetAssetDisplayName() const override { return NSLOCTEXT("AssetTypeActions", "AssetTypeActions_MediaPlate", "Media Plate"); }
+	virtual FLinearColor GetAssetColor() const override { return FLinearColor::Red; }
+	virtual TSoftClassPtr<UObject> GetAssetClass() const override { return UMediaPlateComponent::StaticClass(); }
+	virtual TConstArrayView<FAssetCategoryPath> GetAssetCategories() const override
+	{
+		static const auto Categories = { EAssetCategoryPaths::Media };
+		return Categories;
+	}
+	virtual EAssetCommandResult OpenAssets(const FAssetOpenArgs& OpenArgs) const override;
+	// UAssetDefinition End
+};
diff --git a/Engine/Plugins/Media/MediaPlate/Source/MediaPlateEditor/Private/AssetTools/MediaPlateActions.cpp b/Engine/Plugins/Media/MediaPlate/Source/MediaPlateEditor/Private/AssetTools/MediaPlateActions.cpp
deleted file mode 100644
diff --git a/Engine/Plugins/Media/MediaPlate/Source/MediaPlateEditor/Private/AssetTools/MediaPlateActions.h b/Engine/Plugins/Media/MediaPlate/Source/MediaPlateEditor/Private/AssetTools/MediaPlateActions.h
deleted file mode 100644
diff --git a/Engine/Plugins/Media/MediaPlate/Source/MediaPlateEditor/Private/MediaPlateCustomization.cpp b/Engine/Plugins/Media/MediaPlate/Source/MediaPlateEditor/Private/MediaPlateCustomization.cpp
--- a/Engine/Plugins/Media/MediaPlate/Source/MediaPlateEditor/Private/MediaPlateCustomization.cpp
+++ b/Engine/Plugins/Media/MediaPlate/Source/MediaPlateEditor/Private/MediaPlateCustomization.cpp
@@ -13,6 +13,7 @@
 #include "MediaPlate.h"
 #include "MediaPlateComponent.h"
 #include "MediaPlateEditorModule.h"
+#include "MediaPlateEditorStyle.h"
 #include "MediaPlayer.h"
 #include "MediaPlaylist.h"
 #include "MediaSource.h"
@@ -56,16 +57,7 @@ void FMediaPlateCustomization::CustomizeDetails(IDetailLayoutBuilder& DetailBuil
 	}
 
 	// Get style.
-	const ISlateStyle* Style = nullptr;
-	FMediaPlateEditorModule* EditorModule = FModuleManager::LoadModulePtr<FMediaPlateEditorModule>("MediaPlateEditor");
-	if (EditorModule != nullptr)
-	{
-		Style = EditorModule->GetStyle().Get();
-	}
-	if (Style == nullptr)
-	{
-		Style = &FAppStyle::Get();
-	}
+	const ISlateStyle* Style = &FMediaPlateEditorStyle::Get().Get();
 
 	// Hide the static mesh.
 	IDetailCategoryBuilder& StaticMeshCategory = DetailBuilder.EditCategory("StaticMesh");
diff --git a/Engine/Plugins/Media/MediaPlate/Source/MediaPlateEditor/Private/MediaPlateEditorModule.cpp b/Engine/Plugins/Media/MediaPlate/Source/MediaPlateEditor/Private/MediaPlateEditorModule.cpp
--- a/Engine/Plugins/Media/MediaPlate/Source/MediaPlateEditor/Private/MediaPlateEditorModule.cpp
+++ b/Engine/Plugins/Media/MediaPlate/Source/MediaPlateEditor/Private/MediaPlateEditorModule.cpp
@@ -2,7 +2,6 @@
 
 #include "MediaPlateEditorModule.h"
 
-#include "AssetTools/MediaPlateActions.h"
 #include "Editor.h"
 #include "ISequencerModule.h"
 #include "LevelEditor.h"
@@ -25,14 +24,14 @@
 
 DEFINE_LOG_CATEGORY(LogMediaPlateEditor);
 
+TSharedPtr<FMediaPlateEditorStyle> FMediaPlateEditorStyle::Singleton;
+
 void FMediaPlateEditorModule::StartupModule()
 {
-	Style = MakeShareable(new FMediaPlateEditorStyle());
+	FMediaPlateEditorStyle::Get();
 
 	FMediaPlateEditorCommands::Register();
 
-	RegisterAssetTools();
-
 	// Register customizations.
 	MediaPlateName = UMediaPlateComponent::StaticClass()->GetFName();
 
@@ -72,8 +71,6 @@ void FMediaPlateEditorModule::ShutdownModule()
 		SequencerModulePtr->UnRegisterTrackEditor(TrackEditorBindingHandle);
 	}
 
-	UnregisterAssetTools();
-
 	// Unregister customizations.
 	FPropertyEditorModule& PropertyModule = FModuleManager::LoadModuleChecked<FPropertyEditorModule>("PropertyEditor");
 	PropertyModule.UnregisterCustomClassLayout(MediaPlateName);
@@ -138,34 +135,6 @@ bool FMediaPlateEditorModule::RemoveMediaSourceFromDragDropCache(UMediaSource* M
 	return bIsInCache;
 }
 
-void FMediaPlateEditorModule::RegisterAssetTools()
-{
-	IAssetTools& AssetTools = FModuleManager::LoadModuleChecked<FAssetToolsModule>("AssetTools").Get();
-
-	RegisterAssetTypeAction(AssetTools, MakeShareable(new FMediaPlateActions(Style.ToSharedRef())));
-}
-
-void FMediaPlateEditorModule::RegisterAssetTypeAction(IAssetTools& AssetTools, TSharedRef<IAssetTypeActions> Action)
-{
-	AssetTools.RegisterAssetTypeActions(Action);
-	RegisteredAssetTypeActions.Add(Action);
-}
-
-void FMediaPlateEditorModule::UnregisterAssetTools()
-{
-	FAssetToolsModule* AssetToolsModule = FModuleManager::GetModulePtr<FAssetToolsModule>("AssetTools");
-
-	if (AssetToolsModule != nullptr)
-	{
-		IAssetTools& AssetTools = AssetToolsModule->Get();
-
-		for (TSharedRef<IAssetTypeActions>& Action : RegisteredAssetTypeActions)
-		{
-			AssetTools.UnregisterAssetTypeActions(Action);
-		}
-	}
-}
-
 void FMediaPlateEditorModule::OnPostEngineInit()
 {
 	UEditorAssetSubsystem* EditorAssetSubsystem = GEditor->GetEditorSubsystem<UEditorAssetSubsystem>();
diff --git a/Engine/Plugins/Media/MediaPlate/Source/MediaPlateEditor/Private/MediaPlateEditorStyle.h b/Engine/Plugins/Media/MediaPlate/Source/MediaPlateEditor/Private/MediaPlateEditorStyle.h
--- a/Engine/Plugins/Media/MediaPlate/Source/MediaPlateEditor/Private/MediaPlateEditorStyle.h
+++ b/Engine/Plugins/Media/MediaPlate/Source/MediaPlateEditor/Private/MediaPlateEditorStyle.h
@@ -15,4 +15,16 @@ public:
 
 	 /** Destructor. */
 	~FMediaPlateEditorStyle();
+
+	static TSharedRef<FMediaPlateEditorStyle> Get()
+	{
+		if (!Singleton.IsValid())
+		{
+			Singleton = MakeShareable(new FMediaPlateEditorStyle);
+		}
+		return Singleton.ToSharedRef();
+	}
+
+private:
+	static TSharedPtr<FMediaPlateEditorStyle> Singleton;
 };
diff --git a/Engine/Plugins/Media/MediaPlate/Source/MediaPlateEditor/Public/MediaPlateEditorModule.h b/Engine/Plugins/Media/MediaPlate/Source/MediaPlateEditor/Public/MediaPlateEditorModule.h
--- a/Engine/Plugins/Media/MediaPlate/Source/MediaPlateEditor/Public/MediaPlateEditorModule.h
+++ b/Engine/Plugins/Media/MediaPlate/Source/MediaPlateEditor/Public/MediaPlateEditorModule.h
@@ -10,9 +10,6 @@
 #include "UObject/ObjectPtr.h"
 
 class FDragDropEvent;
-class IAssetTools;
-class IAssetTypeActions;
-class ISlateStyle;
 class UMediaPlateComponent;
 class UMediaSource;
 struct FAssetData;
@@ -32,11 +29,6 @@ public:
 	virtual void Tick(float DeltaTime) override;
 	virtual TStatId GetStatId() const override { RETURN_QUICK_DECLARE_CYCLE_STAT(FMediaPlateEditorModule, STATGROUP_Tickables); }
 	
-	/**
-	 * Get the style used by this module.
-	 **/
-	TSharedPtr<ISlateStyle> GetStyle() { return Style; }
-	
 	/**
 	 * Call this when a media plate starts playing so we can track it.
 	 */
@@ -59,10 +51,6 @@ private:
 	/** Holds all the media plates that are playing. */
 	TArray<TWeakObjectPtr<UMediaPlateComponent>> ActiveMediaPlates;
 
-	/** Holds the plug-ins style set. */
-	TSharedPtr<ISlateStyle> Style;
-	/** The collection of registered asset type actions. */
-	TArray<TSharedRef<IAssetTypeActions>> RegisteredAssetTypeActions;
 	/** Handle for our track editor. */
 	FDelegateHandle TrackEditorBindingHandle;
 
@@ -71,24 +59,6 @@ private:
 	/** Caches if real time viewports are enabled. */
 	bool bIsRealTimeViewportsEnabled = false;
 
-	/**
-	 * Registers all of our asset tools.
-	 */
-	void RegisterAssetTools();
-
-	/**
-	 * Registers a single asset type action.
-	 *
-	 * @param AssetTools	The asset tools object to register with.
-	 * @param Action		The asset type action to register.
-	 */
-	void RegisterAssetTypeAction(IAssetTools& AssetTools, TSharedRef<IAssetTypeActions> Action);
-
-	/**
-	 * Unregisters all of our asset tools.
-	 */
-	void UnregisterAssetTools();
-
 	/**
 	 * Called after the engine has initialized.
 	 */
