MediaPlate: Fixed media plates made in UE 5.1 not showing playback.

diff --git a/Engine/Plugins/Media/MediaPlate/Source/MediaPlate/Private/MediaPlateComponent.cpp b/Engine/Plugins/Media/MediaPlate/Source/MediaPlate/Private/MediaPlateComponent.cpp
--- a/Engine/Plugins/Media/MediaPlate/Source/MediaPlate/Private/MediaPlateComponent.cpp
+++ b/Engine/Plugins/Media/MediaPlate/Source/MediaPlate/Private/MediaPlateComponent.cpp
@@ -90,6 +90,23 @@ UMediaPlateComponent::UMediaPlateComponent(const FObjectInitializer& ObjectIniti
 	VisibleMipsTilesCalculations = EMediaTextureVisibleMipsTiles::Plane;
 }
 
+#if WITH_EDITOR
+void UMediaPlateComponent::PostLoad()
+{
+	Super::PostLoad();
+
+	// Use the old media texture if we have one.
+	if (MediaTexture_DEPRECATED != nullptr)
+	{
+		if (MediaTextures.Num() == 0)
+		{
+			MediaTextures.Add(MediaTexture_DEPRECATED);
+		}
+		MediaTexture_DEPRECATED = nullptr;
+	}
+}
+#endif // WITH_EDITOR
+
 void UMediaPlateComponent::OnRegister()
 {
 	Super::OnRegister();
diff --git a/Engine/Plugins/Media/MediaPlate/Source/MediaPlate/Public/MediaPlateComponent.h b/Engine/Plugins/Media/MediaPlate/Source/MediaPlate/Public/MediaPlateComponent.h
--- a/Engine/Plugins/Media/MediaPlate/Source/MediaPlate/Public/MediaPlateComponent.h
+++ b/Engine/Plugins/Media/MediaPlate/Source/MediaPlate/Public/MediaPlateComponent.h
@@ -36,6 +36,9 @@ class MEDIAPLATE_API UMediaPlateComponent : public UActorComponent,
 
 public:
 	//~ UActorComponent interface.
+#if WITH_EDITOR
+	virtual void PostLoad() override;
+#endif // WITH_EDITOR
 	virtual void OnRegister() override;
 	virtual void BeginPlay() override;
 	virtual void BeginDestroy() override;
@@ -328,6 +331,12 @@ private:
 	/** Name for our playlist. */
 	static FLazyName MediaPlaylistName;
 
+#if WITH_EDITORONLY_DATA
+	/** Superseded by MediaTextures. */
+	UPROPERTY(Instanced)
+	TObjectPtr<UMediaTexture> MediaTexture_DEPRECATED;
+#endif // WITH_EDITORONLY_DATA
+
 	/** Holds the media textures. */
 	UPROPERTY(Instanced)
 	TArray<TObjectPtr<UMediaTexture>> MediaTextures;
