MediaPlate: Fixed auto aspect ratio not working with Electra.

diff --git a/Engine/Plugins/Media/MediaPlate/Source/MediaPlate/Private/MediaPlateComponent.cpp b/Engine/Plugins/Media/MediaPlate/Source/MediaPlate/Private/MediaPlateComponent.cpp
--- a/Engine/Plugins/Media/MediaPlate/Source/MediaPlate/Private/MediaPlateComponent.cpp
+++ b/Engine/Plugins/Media/MediaPlate/Source/MediaPlate/Private/MediaPlateComponent.cpp
@@ -477,7 +477,7 @@ bool UMediaPlateComponent::PlayMediaSource(UMediaSource* InMediaSource, bool bIn
 
 void UMediaPlateComponent::TryActivateAspectRatioAuto()
 {
-	if ((MediaPlayer != nullptr) && (MediaPlayer->IsClosed() == false))
+	if (MediaPlayer != nullptr)
 	{
 		// Are we using automatic aspect ratio?
 		if ((bIsAspectRatioAuto) &&
@@ -586,7 +586,7 @@ void UMediaPlateComponent::TickOutput()
 	if (MediaPlayer != nullptr)
 	{
 		// Is the player ready?
-		if (MediaPlayer->IsPreparing() == false)
+		if ((MediaPlayer->IsClosed() == false) && (MediaPlayer->IsPreparing() == false))
 		{
 			FIntPoint VideoDim = MediaPlayer->GetVideoTrackDimensions(INDEX_NONE, INDEX_NONE);
 			if (VideoDim.Y != 0)
