MediaPlate: Fixed assert when in Sequencer deleting a section from the top layer and then moving a section from the bottom layer to be active.

diff --git a/Engine/Plugins/Media/MediaPlate/Source/MediaPlate/Private/MediaPlateComponent.cpp b/Engine/Plugins/Media/MediaPlate/Source/MediaPlate/Private/MediaPlateComponent.cpp
--- a/Engine/Plugins/Media/MediaPlate/Source/MediaPlate/Private/MediaPlateComponent.cpp
+++ b/Engine/Plugins/Media/MediaPlate/Source/MediaPlate/Private/MediaPlateComponent.cpp
@@ -697,7 +697,7 @@ void UMediaPlateComponent::ProxySetTextureBlend(int32 LayerIndex, int32 TextureI
 			if (MID != nullptr)
 			{
 				int32 MatNumLayers = MediaTextures.Num() / MatNumTexPerLayer;
-				if (LayerIndex < MatNumLayers)
+				if ((LayerIndex < MatNumLayers) && (LayerIndex < TextureLayers.Num()))
 				{
 					const TArray<int32>& Layer = TextureLayers[LayerIndex];
 					for (int32 LayerTexIndex = 0;
