MediaPlate: Fixed assert when there aren't enough textures for crossfades.
#rnx

diff --git a/Engine/Plugins/Media/MediaPlate/Source/MediaPlate/Private/MediaPlateComponent.cpp b/Engine/Plugins/Media/MediaPlate/Source/MediaPlate/Private/MediaPlateComponent.cpp
--- a/Engine/Plugins/Media/MediaPlate/Source/MediaPlate/Private/MediaPlateComponent.cpp
+++ b/Engine/Plugins/Media/MediaPlate/Source/MediaPlate/Private/MediaPlateComponent.cpp
@@ -644,29 +644,31 @@ UMediaSource* UMediaPlateComponent::ProxyGetMediaSourceFromIndex(int32 Index) co
 UMediaTexture* UMediaPlateComponent::ProxyGetMediaTexture(int32 LayerIndex, int32 TextureIndex)
 {
 	UMediaTexture* MediaTexture = GetMediaTexture(TextureIndex);
-
-	if (TextureLayers.Num() < LayerIndex + 1)
+	if (MediaTexture != nullptr)
 	{
-		TextureLayers.SetNum(LayerIndex + 1);
-	}
+		if (TextureLayers.Num() < LayerIndex + 1)
+		{
+			TextureLayers.SetNum(LayerIndex + 1);
+		}
 
-	// Fill up an empty slot if there is one.
-	bool bIsTextureSet = false;
-	for (int32 Index = 0; Index < TextureLayers[LayerIndex].Num(); ++Index)
-	{
-		if (TextureLayers[LayerIndex][Index] < 0)
+		// Fill up an empty slot if there is one.
+		bool bIsTextureSet = false;
+		for (int32 Index = 0; Index < TextureLayers[LayerIndex].Num(); ++Index)
 		{
-			TextureLayers[LayerIndex][Index] = TextureIndex;
-			bIsTextureSet = true;
-			break;
+			if (TextureLayers[LayerIndex][Index] < 0)
+			{
+				TextureLayers[LayerIndex][Index] = TextureIndex;
+				bIsTextureSet = true;
+				break;
+			}
+		}
+		if (bIsTextureSet == false)
+		{
+			TextureLayers[LayerIndex].Add(TextureIndex);
 		}
-	}
-	if (bIsTextureSet == false)
-	{
-		TextureLayers[LayerIndex].Add(TextureIndex);
-	}
 
-	UpdateTextureLayers();
+		UpdateTextureLayers();
+	}
 
 	return MediaTexture;
 }
