MediaPlate: Fixed broken playback when no longer using this in Sequencer and then attempting playback.
#rnx

diff --git a/Engine/Plugins/Media/MediaPlate/Source/MediaPlate/Private/MediaPlateComponent.cpp b/Engine/Plugins/Media/MediaPlate/Source/MediaPlate/Private/MediaPlateComponent.cpp
--- a/Engine/Plugins/Media/MediaPlate/Source/MediaPlate/Private/MediaPlateComponent.cpp
+++ b/Engine/Plugins/Media/MediaPlate/Source/MediaPlate/Private/MediaPlateComponent.cpp
@@ -231,6 +231,7 @@ void UMediaPlateComponent::Open()
 	bIsMediaPlatePlaying = true;
 	CurrentRate = bPlayOnOpen ? 1.0f : 0.0f;
 	PlaylistIndex = 0;
+	SetNormalMode(true);
 
 	if (IsVisible())
 	{
@@ -646,6 +647,7 @@ UMediaTexture* UMediaPlateComponent::ProxyGetMediaTexture(int32 LayerIndex, int3
 	UMediaTexture* MediaTexture = GetMediaTexture(TextureIndex);
 	if (MediaTexture != nullptr)
 	{
+		SetNormalMode(false);
 		if (TextureLayers.Num() < LayerIndex + 1)
 		{
 			TextureLayers.SetNum(LayerIndex + 1);
@@ -677,16 +679,19 @@ void UMediaPlateComponent::ProxyReleaseMediaTexture(int32 LayerIndex, int32 Text
 {
 	ProxySetTextureBlend(LayerIndex, TextureIndex, 0.0f);
 
-	for (int32 Index = 0; Index < TextureLayers[LayerIndex].Num(); ++Index)
+	if (LayerIndex < TextureLayers.Num())
 	{
-		if (TextureLayers[LayerIndex][Index] == TextureIndex)
+		for (int32 Index = 0; Index < TextureLayers[LayerIndex].Num(); ++Index)
 		{
-			TextureLayers[LayerIndex][Index] = -1;
-			break;
+			if (TextureLayers[LayerIndex][Index] == TextureIndex)
+			{
+				TextureLayers[LayerIndex][Index] = -1;
+				break;
+			}
 		}
+
+		UpdateTextureLayers();
 	}
-	
-	UpdateTextureLayers();
 }
 
 void UMediaPlateComponent::ProxySetTextureBlend(int32 LayerIndex, int32 TextureIndex, float Blend)
@@ -978,6 +983,40 @@ void UMediaPlateComponent::SetUpTextures(UE::MediaPlateComponent::ESetUpTextures
 	}
 }
 
+
+void UMediaPlateComponent::SetNormalMode(bool bInIsNormalMode)
+{
+#if WITH_EDITOR
+	// Switching between normal mode and proxy mode should only be needed in the editor.
+	if (bIsNormalMode != bInIsNormalMode)
+	{
+		bIsNormalMode = bInIsNormalMode;
+		if (bIsNormalMode)
+		{
+			// Only want 1 texture.
+			if (TextureLayers.Num() != 1)
+			{
+				TextureLayers.SetNum(1);
+			}
+			if (TextureLayers[0].Num() != 1)
+			{
+				TextureLayers[0].SetNum(1);
+			}
+			TextureLayers[0][0] = 0;
+			UpdateTextureLayers();
+			
+			ProxySetTextureBlend(0, 0, 1.0f);
+			MediaTextures[0]->SetMediaPlayer(MediaPlayer);
+		}
+		else
+		{
+			// Proxy will set these up.
+			TextureLayers.Reset();
+		}
+	}
+#endif // WITH_EDITOR
+}
+
 void UMediaPlateComponent::UpdateTextureLayers()
 {
 	if (AMediaPlate* MediaPlate = GetOwner<AMediaPlate>())
diff --git a/Engine/Plugins/Media/MediaPlate/Source/MediaPlate/Public/MediaPlateComponent.h b/Engine/Plugins/Media/MediaPlate/Source/MediaPlate/Public/MediaPlateComponent.h
--- a/Engine/Plugins/Media/MediaPlate/Source/MediaPlate/Public/MediaPlateComponent.h
+++ b/Engine/Plugins/Media/MediaPlate/Source/MediaPlate/Public/MediaPlateComponent.h
@@ -356,6 +356,11 @@ private:
 	/** True if we should resume where we left off when we open the media. */
 	bool bResumeWhenOpened = false;
 
+#if WITH_EDITOR
+	/** True if we are in normal mode (as opposed to proxy mode). */
+	bool bIsNormalMode = false;
+#endif
+
 	/**
 	 * Contains all of our layers.
 	 * Each layer contains which textures it has.
@@ -436,6 +441,11 @@ private:
 	 */
 	void SetUpTextures(UE::MediaPlateComponent::ESetUpTexturesFlags Flags);
 
+	/**
+	 * Sets either normal mode or proxy mode for something like Sequencer.
+	 */
+	void SetNormalMode(bool bInIsNormalMode);
+
 	/**
 	 * Sets textures in our material according to the layer assignments.
 	 */
