Media Plate: Mesh selection property wasn't transacting over ndisplay

diff --git a/Engine/Plugins/Media/MediaPlate/Source/MediaPlateEditor/Private/MediaPlateCustomization.cpp b/Engine/Plugins/Media/MediaPlate/Source/MediaPlateEditor/Private/MediaPlateCustomization.cpp
--- a/Engine/Plugins/Media/MediaPlate/Source/MediaPlateEditor/Private/MediaPlateCustomization.cpp
+++ b/Engine/Plugins/Media/MediaPlate/Source/MediaPlateEditor/Private/MediaPlateCustomization.cpp
@@ -18,6 +18,7 @@
 #include "MediaPlaylist.h"
 #include "MediaSource.h"
 #include "PropertyCustomizationHelpers.h"
+#include "ScopedTransaction.h"
 #include "Styling/AppStyle.h"
 #include "Subsystems/AssetEditorSubsystem.h"
 #include "ThumbnailRendering/ThumbnailManager.h"
@@ -757,6 +758,8 @@ void FMediaPlateCustomization::SetMeshMode(EMediaTextureVisibleMipsTiles InMode)
 {
 	if (MeshMode != InMode)
 	{
+		const FScopedTransaction Transaction(LOCTEXT("SetMeshMode", "Media Plate Mesh Changed"));
+
 		MeshMode = InMode;
 		for (const TWeakObjectPtr<UMediaPlateComponent>& MediaPlatePtr : MediaPlatesList)
 		{
@@ -1030,6 +1033,8 @@ FString FMediaPlateCustomization::GetStaticMeshPath() const
 
 void FMediaPlateCustomization::OnStaticMeshChanged(const FAssetData& AssetData)
 {
+	const FScopedTransaction Transaction(LOCTEXT("OnStaticMeshChanged", "Media Plate Custom Mesh Changed"));
+
 	// Update the static mesh.
 	UStaticMesh* StaticMesh = Cast<UStaticMesh>(AssetData.GetAsset());
 	for (TWeakObjectPtr<UMediaPlateComponent>& MediaPlatePtr : MediaPlatesList)
diff --git a/Engine/Plugins/Media/MediaPlate/Source/MediaPlateEditor/Private/MediaPlateCustomizationMesh.cpp b/Engine/Plugins/Media/MediaPlate/Source/MediaPlateEditor/Private/MediaPlateCustomizationMesh.cpp
--- a/Engine/Plugins/Media/MediaPlate/Source/MediaPlateEditor/Private/MediaPlateCustomizationMesh.cpp
+++ b/Engine/Plugins/Media/MediaPlate/Source/MediaPlateEditor/Private/MediaPlateCustomizationMesh.cpp
@@ -107,6 +107,9 @@ void FMediaPlateCustomizationMesh::SetMesh(UStaticMeshComponent* StaticMeshCompo
 {
 	if (StaticMeshComponent != nullptr)
 	{
+		FProperty* StaticMeshProperty = FindFieldChecked<FProperty>(UStaticMeshComponent::StaticClass(), "StaticMesh");
+		StaticMeshComponent->PreEditChange(StaticMeshProperty);
+
 		// Get existing mesh.
 		TObjectPtr<UStaticMesh> OldMesh = StaticMeshComponent->GetStaticMesh();
 		if (OldMesh != nullptr)
@@ -130,7 +133,6 @@ void FMediaPlateCustomizationMesh::SetMesh(UStaticMeshComponent* StaticMeshCompo
 		StaticMeshComponent->SetRelativeScale3D(FVector::OneVector);
 
 		// Call PostEditChangeProperty so it updates properly.
-		FProperty* StaticMeshProperty = FindFieldChecked<FProperty>(UStaticMeshComponent::StaticClass(), "StaticMesh");
 		FPropertyChangedEvent PropertyEvent(StaticMeshProperty);
 		StaticMeshComponent->PostEditChangeProperty(PropertyEvent);
 
