MediaCompositing: Sets correct aspect ratio when using media plate.
#rnx

diff --git a/Engine/Plugins/Media/MediaCompositing/Source/MediaCompositing/Private/MovieScene/MovieSceneMediaData.cpp b/Engine/Plugins/Media/MediaCompositing/Source/MediaCompositing/Private/MovieScene/MovieSceneMediaData.cpp
--- a/Engine/Plugins/Media/MediaCompositing/Source/MediaCompositing/Private/MovieScene/MovieSceneMediaData.cpp
+++ b/Engine/Plugins/Media/MediaCompositing/Source/MediaCompositing/Private/MovieScene/MovieSceneMediaData.cpp
@@ -14,7 +14,8 @@
  *****************************************************************************/
 
 FMovieSceneMediaData::FMovieSceneMediaData()
-	: bOverrideMediaPlayer(false)
+	: bIsAspectRatioSet(false)
+	, bOverrideMediaPlayer(false)
 	, MediaPlayer(nullptr)
 	, ProxyTextureIndex(0)
 	, SeekOnOpenTime(FTimespan::MinValue())
diff --git a/Engine/Plugins/Media/MediaCompositing/Source/MediaCompositing/Private/MovieScene/MovieSceneMediaTemplate.cpp b/Engine/Plugins/Media/MediaCompositing/Source/MediaCompositing/Private/MovieScene/MovieSceneMediaTemplate.cpp
--- a/Engine/Plugins/Media/MediaCompositing/Source/MediaCompositing/Private/MovieScene/MovieSceneMediaTemplate.cpp
+++ b/Engine/Plugins/Media/MediaCompositing/Source/MediaCompositing/Private/MovieScene/MovieSceneMediaTemplate.cpp
@@ -148,6 +148,14 @@ struct FMediaSectionExecutionToken
 				{
 					return;
 				}
+
+				if (SectionData.bIsAspectRatioSet == false)
+				{
+					if (PlayerProxyInterface->ProxySetAspectRatio(MediaPlayer))
+					{
+						SectionData.bIsAspectRatioSet = true;
+					}
+				}
 			}
 		}
 
diff --git a/Engine/Plugins/Media/MediaCompositing/Source/MediaCompositing/Public/MovieSceneMediaData.h b/Engine/Plugins/Media/MediaCompositing/Source/MediaCompositing/Public/MovieSceneMediaData.h
--- a/Engine/Plugins/Media/MediaCompositing/Source/MediaCompositing/Public/MovieSceneMediaData.h
+++ b/Engine/Plugins/Media/MediaCompositing/Source/MediaCompositing/Public/MovieSceneMediaData.h
@@ -68,6 +68,11 @@ public:
 	 */
 	void TearDown();
 
+	/**
+	 * Stores if the aspect ratio has been set yet.
+	 */
+	bool bIsAspectRatioSet;
+
 private:
 	/**
 	 * Does the work needed so we can use our proxy media texture.
diff --git a/Engine/Plugins/Media/MediaPlate/Source/MediaPlate/Private/MediaPlateComponent.cpp b/Engine/Plugins/Media/MediaPlate/Source/MediaPlate/Private/MediaPlateComponent.cpp
--- a/Engine/Plugins/Media/MediaPlate/Source/MediaPlate/Private/MediaPlateComponent.cpp
+++ b/Engine/Plugins/Media/MediaPlate/Source/MediaPlate/Private/MediaPlateComponent.cpp
@@ -481,8 +481,7 @@ void UMediaPlateComponent::TryActivateAspectRatioAuto()
 	if (MediaPlayer != nullptr)
 	{
 		// Are we using automatic aspect ratio?
-		if ((bIsAspectRatioAuto) &&
-			(VisibleMipsTilesCalculations == EMediaTextureVisibleMipsTiles::Plane))
+		if (IsAspectRatioAutoAllowed())
 		{
 			// Start the clock sink so we can tick.
 			IMediaModule* MediaModule = FModuleManager::LoadModulePtr<IMediaModule>("Media");
@@ -498,6 +497,11 @@ void UMediaPlateComponent::TryActivateAspectRatioAuto()
 	}
 }
 
+bool UMediaPlateComponent::IsAspectRatioAutoAllowed()
+{
+	return ((bIsAspectRatioAuto) &&
+		(VisibleMipsTilesCalculations == EMediaTextureVisibleMipsTiles::Plane));
+}
 
 float UMediaPlateComponent::GetAspectRatio()
 {
@@ -584,22 +588,10 @@ void UMediaPlateComponent::SetNumberOfTextures(int32 NumTextures)
 
 void UMediaPlateComponent::TickOutput()
 {
-	if (MediaPlayer != nullptr)
+	if (ProxySetAspectRatio(MediaPlayer))
 	{
-		// Is the player ready?
-		if ((MediaPlayer->IsClosed() == false) && (MediaPlayer->IsPreparing() == false))
-		{
-			FIntPoint VideoDim = MediaPlayer->GetVideoTrackDimensions(INDEX_NONE, INDEX_NONE);
-			if (VideoDim.Y != 0)
-			{ 
-				// Set aspect ratio.
-				float AspectRatio = (float)VideoDim.X / (float)VideoDim.Y;
-				SetAspectRatio(AspectRatio);
-					
-				// No need to tick anymore.
-				StopClockSink();
-			}
-		}
+		// No need to tick anymore.
+		StopClockSink();
 	}
 }
 
@@ -694,6 +686,34 @@ void UMediaPlateComponent::ProxyReleaseMediaTexture(int32 LayerIndex, int32 Text
 	}
 }
 
+bool UMediaPlateComponent::ProxySetAspectRatio(UMediaPlayer* InMediaPlayer)
+{
+	bool bIsDone = false;
+
+	if (IsAspectRatioAutoAllowed())
+	{
+		// Is the player ready?
+		if ((InMediaPlayer != nullptr) && (InMediaPlayer->IsClosed() == false) &&
+			(InMediaPlayer->IsPreparing() == false))
+		{
+			FIntPoint VideoDim = InMediaPlayer->GetVideoTrackDimensions(INDEX_NONE, INDEX_NONE);
+			if (VideoDim.Y != 0)
+			{
+				// Set aspect ratio.
+				float AspectRatio = (float)VideoDim.X / (float)VideoDim.Y;
+				SetAspectRatio(AspectRatio);
+				bIsDone = true;
+			}
+		}
+	}
+	else
+	{
+		bIsDone = true;
+	}
+
+	return bIsDone;
+}
+
 void UMediaPlateComponent::ProxySetTextureBlend(int32 LayerIndex, int32 TextureIndex, float Blend)
 {
 	if (AMediaPlate* MediaPlate = GetOwner<AMediaPlate>())
diff --git a/Engine/Plugins/Media/MediaPlate/Source/MediaPlate/Public/MediaPlateComponent.h b/Engine/Plugins/Media/MediaPlate/Source/MediaPlate/Public/MediaPlateComponent.h
--- a/Engine/Plugins/Media/MediaPlate/Source/MediaPlate/Public/MediaPlateComponent.h
+++ b/Engine/Plugins/Media/MediaPlate/Source/MediaPlate/Public/MediaPlateComponent.h
@@ -257,6 +257,7 @@ public:
 	virtual UMediaSource* ProxyGetMediaSourceFromIndex(int32 Index) const override;
 	virtual UMediaTexture* ProxyGetMediaTexture(int32 LayerIndex, int32 TextureIndex) override;
 	virtual void ProxyReleaseMediaTexture(int32 LayerIndex, int32 TextureIndex) override;
+	virtual bool ProxySetAspectRatio(UMediaPlayer* InMediaPlayer) override;
 	virtual void ProxySetTextureBlend(int32 LayerIndex, int32 TextureIndex, float Blend) override;
 
 private:
@@ -384,6 +385,11 @@ private:
 	 */
 	void TryActivateAspectRatioAuto();
 
+	/**
+	 * Returns true if auto aspect ratio is enabled and our mesh supports this (e.g. planar).
+	 */
+	bool IsAspectRatioAutoAllowed();
+
 	/**
 	 * Stops the clock sink so we no longer tick.
 	 */
diff --git a/Engine/Source/Runtime/MediaAssets/Public/MediaPlayerProxyInterface.h b/Engine/Source/Runtime/MediaAssets/Public/MediaPlayerProxyInterface.h
--- a/Engine/Source/Runtime/MediaAssets/Public/MediaPlayerProxyInterface.h
+++ b/Engine/Source/Runtime/MediaAssets/Public/MediaPlayerProxyInterface.h
@@ -6,6 +6,7 @@
 
 #include "MediaPlayerProxyInterface.generated.h"
 
+class UMediaPlayer;
 class UMediaSource;
 class UMediaTexture;
 struct FMediaSourceCacheSettings;
@@ -68,6 +69,18 @@ public:
 	 */
 	virtual void ProxyReleaseMediaTexture(int32 LayerIndex, int32 TextureIndex) = 0;
 
+	/**
+	 * Sets the aspect ratio of the proxy based on what the media player is playing.
+	 * 
+	 * This should be called every frame until it returns true.
+	 * It might take a few frames to discover what the aspect ratio is.
+	 * Also some proxies might not support aspect ratios,
+	 * so they will also return true so you do not call this endlessly.
+	 *
+	 * @return True if you no longer need to call this.
+	 */
+	virtual bool ProxySetAspectRatio(UMediaPlayer* InMediaPlayer) = 0;
+
 	/**
 	 * Set the blend value for a texture that was retrieved from ProxyGetMediaTexture.
 	 */
