MediaPlate: LWC truncation fixes.

#rnx

diff --git a/Engine/Plugins/Media/MediaPlate/Source/MediaPlate/Private/MediaPlateComponent.cpp b/Engine/Plugins/Media/MediaPlate/Source/MediaPlate/Private/MediaPlateComponent.cpp
--- a/Engine/Plugins/Media/MediaPlate/Source/MediaPlate/Private/MediaPlateComponent.cpp
+++ b/Engine/Plugins/Media/MediaPlate/Source/MediaPlate/Private/MediaPlateComponent.cpp
@@ -191,7 +191,7 @@ void UMediaPlateComponent::TickComponent(float DeltaTime, enum ELevelTick TickTy
 				if (MediaPlayer->IsPlaying())
 				{
 					MediaPlayer->Pause();
-					TimeWhenPlaybackPaused = FApp::GetGameTime();
+					TimeWhenPlaybackPaused = static_cast<float>(FApp::GetGameTime());
 				}
 			}
 		}
@@ -255,7 +255,7 @@ void UMediaPlateComponent::Open()
 	else
 	{
 		bWantsToPlayWhenVisible = true;
-		TimeWhenPlaybackPaused = FApp::GetGameTime();
+		TimeWhenPlaybackPaused = static_cast<float>(FApp::GetGameTime());
 	}
 
 	UpdateTicking();
@@ -512,7 +512,7 @@ float UMediaPlateComponent::GetAspectRatio()
 		float AspectRatio = 0.0f;
 		if (Scale.Z != 0.0f)
 		{
-			AspectRatio = Scale.Y / Scale.Z;
+			AspectRatio = static_cast<float>(Scale.Y / Scale.Z);
 		}
 		return AspectRatio;
 	}
@@ -817,7 +817,7 @@ FTimespan UMediaPlateComponent::GetResumeTime()
 	if (MediaPlayer != nullptr)
 	{
 		PlayerTime = MediaPlayer->GetTime();
-		float CurrentTime = FApp::GetGameTime();
+		float CurrentTime = static_cast<float>(FApp::GetGameTime());
 		float ElapsedTime = CurrentTime - TimeWhenPlaybackPaused;
 		PlayerTime += FTimespan::FromSeconds(ElapsedTime);
 		
diff --git a/Engine/Plugins/Media/MediaPlate/Source/MediaPlateEditor/Private/MediaPlateCustomization.cpp b/Engine/Plugins/Media/MediaPlate/Source/MediaPlateEditor/Private/MediaPlateCustomization.cpp
--- a/Engine/Plugins/Media/MediaPlate/Source/MediaPlateEditor/Private/MediaPlateCustomization.cpp
+++ b/Engine/Plugins/Media/MediaPlate/Source/MediaPlateEditor/Private/MediaPlateCustomization.cpp
@@ -510,7 +510,7 @@ void FMediaPlateCustomization::CustomizeDetails(IDetailLayoutBuilder& DetailBuil
 					.Padding(0, 5, 10, 5)
 					[
 						SNew(SButton)
-							.ContentPadding(3)
+							.ContentPadding(3.0f)
 							.VAlign(VAlign_Center)
 							.HAlign(HAlign_Center)
 							.OnClicked(this, &FMediaPlateCustomization::OnOpenMediaPlate)
@@ -602,7 +602,7 @@ void FMediaPlateCustomization::AddMeshCustomization(IDetailCategoryBuilder& Medi
 				[
 					SNew(SComboButton)
 						.OnGetMenuContent(this, &FMediaPlateCustomization::OnGetAspectRatios)
-						.ContentPadding(2)
+						.ContentPadding(2.0f)
 						.ButtonContent()
 						[
 							SNew(STextBlock)
@@ -644,7 +644,7 @@ void FMediaPlateCustomization::AddMeshCustomization(IDetailCategoryBuilder& Medi
 				[
 					SNew(SComboButton)
 						.OnGetMenuContent(this, &FMediaPlateCustomization::OnGetLetterboxAspectRatios)
-						.ContentPadding(2)
+						.ContentPadding(2.0f)
 						.ButtonContent()
 						[
 							SNew(STextBlock)
@@ -955,7 +955,7 @@ TOptional<float> FMediaPlateCustomization::GetMeshHorizontalRange() const
 		UMediaPlateComponent* MediaPlate = MediaPlatePtr.Get();
 		if (MediaPlate != nullptr)
 		{
-			return MediaPlate->GetMeshRange().X;
+			return static_cast<float>(MediaPlate->GetMeshRange().X);
 		}
 	}
 
@@ -981,7 +981,7 @@ TOptional<float> FMediaPlateCustomization::GetMeshVerticalRange() const
 		UMediaPlateComponent* MediaPlate = MediaPlatePtr.Get();
 		if (MediaPlate != nullptr)
 		{
-			return MediaPlate->GetMeshRange().Y;
+			return static_cast<float>(MediaPlate->GetMeshRange().Y);
 		}
 	}
 
diff --git a/Engine/Plugins/Media/MediaPlate/Source/MediaPlateEditor/Private/MediaPlateSphereGenerator.h b/Engine/Plugins/Media/MediaPlate/Source/MediaPlateEditor/Private/MediaPlateSphereGenerator.h
--- a/Engine/Plugins/Media/MediaPlate/Source/MediaPlateEditor/Private/MediaPlateSphereGenerator.h
+++ b/Engine/Plugins/Media/MediaPlate/Source/MediaPlateEditor/Private/MediaPlateSphereGenerator.h
@@ -95,8 +95,8 @@ private:
 	{
 		// generate the UV's
 		int32 ThetaDivide = bIsClosed ? NumTheta : NumTheta - 1;
-		const float DUVphi = 1.0 / float(NumPhi - 1);
-		const float DUVtheta = 1.0 / float(ThetaDivide);
+		const float DUVphi = 1.0f / float(NumPhi - 1);
+		const float DUVtheta = 1.0f / float(ThetaDivide);
 
 		int32 UVIdx = 0;
 		int32 p,t;
diff --git a/Engine/Plugins/Media/MediaPlate/Source/MediaPlateEditor/Private/Widgets/SMediaPlateEditorMediaDetails.cpp b/Engine/Plugins/Media/MediaPlate/Source/MediaPlateEditor/Private/Widgets/SMediaPlateEditorMediaDetails.cpp
--- a/Engine/Plugins/Media/MediaPlate/Source/MediaPlateEditor/Private/Widgets/SMediaPlateEditorMediaDetails.cpp
+++ b/Engine/Plugins/Media/MediaPlate/Source/MediaPlateEditor/Private/Widgets/SMediaPlateEditorMediaDetails.cpp
@@ -164,8 +164,8 @@ void SMediaPlateEditorMediaDetails::UpdateDetails()
 					: LOCTEXT("MethodStreamed", "Streamed"));
 			NumMips = MediaTexture->GetTextureNumMips();
 			ResourceSize = (MediaTexture->GetResourceSizeBytes(EResourceSizeMode::Exclusive) + 512) / 1024;
-			SurfaceWidth = MediaTexture->GetSurfaceWidth();
-			SurfaceHeight = MediaTexture->GetSurfaceHeight();
+			SurfaceWidth = static_cast<int32>(MediaTexture->GetSurfaceWidth());
+			SurfaceHeight = static_cast<int32>(MediaTexture->GetSurfaceHeight());
 		}
 	}
 
