MediaPlate: Ensure playlist changes transact.

diff --git a/Engine/Plugins/Developer/Concert/ConcertSync/ConcertSyncCore/Config/BaseConcertSyncCore.ini b/Engine/Plugins/Developer/Concert/ConcertSync/ConcertSyncCore/Config/BaseConcertSyncCore.ini
--- a/Engine/Plugins/Developer/Concert/ConcertSync/ConcertSyncCore/Config/BaseConcertSyncCore.ini
+++ b/Engine/Plugins/Developer/Concert/ConcertSync/ConcertSyncCore/Config/BaseConcertSyncCore.ini
@@ -2,7 +2,7 @@
 +IncludeObjectClassFilters=(ObjectOuterClass=None,ObjectClasses=(Class'/Script/Engine.World',Class'/Script/Engine.BookmarkBase',Class'/Script/Engine.MaterialParameterCollection',Class'/Script/MovieScene.MovieSceneSignedObject',Class'/Script/MovieScene.MovieSceneFolder'))
 +IncludeObjectClassFilters=(ObjectOuterClass=Class'/Script/Engine.World',ObjectClasses=(Class'/Script/Engine.Layer',Class'/Script/Engine.Level',Class'/Script/Engine.LevelStreaming',Class'/Script/Engine.Actor',Class'/Script/Engine.ActorComponent',Class'/Script/Engine.Polys',Class'/Script/Engine.BrushBuilder',Class'/Script/Foliage.FoliageType'))
 +IncludeObjectClassFilters=(ObjectOuterClass=Class'/Script/MovieScene.MovieSceneSignedObject',ObjectClasses=(Class'/Script/Engine.Actor',Class'/Script/Engine.ActorComponent'))
-+IncludeObjectClassFilters=(ObjectOuterClass=Class'/Script/Engine.World',ObjectClasses=(Class'/Script/Engine.MaterialInstanceDynamic'))
++IncludeObjectClassFilters=(ObjectOuterClass=Class'/Script/Engine.World',ObjectClasses=(Class'/Script/Engine.MaterialInstanceDynamic', '/Script/MediaAssets.MediaPlaylist'))
 +ExcludeTransactionClassFilters=(ObjectOuterClass=None,ObjectClasses=(Class'/Script/Landscape.LandscapeProxy',Class'/Script/Landscape.LandscapeHeightfieldCollisionComponent',Class'/Script/Landscape.LandscapeComponent'))
 +AllowedTransientProperties="/Script/Engine.Actor:bHiddenEdTemporary"
 +AllowedTransientProperties="/Script/Engine.Actor:bHiddenEdLevel"
diff --git a/Engine/Plugins/Media/MediaPlate/Source/MediaPlateEditor/Private/MediaPlateCustomization.cpp b/Engine/Plugins/Media/MediaPlate/Source/MediaPlateEditor/Private/MediaPlateCustomization.cpp
--- a/Engine/Plugins/Media/MediaPlate/Source/MediaPlateEditor/Private/MediaPlateCustomization.cpp
+++ b/Engine/Plugins/Media/MediaPlate/Source/MediaPlateEditor/Private/MediaPlateCustomization.cpp
@@ -1077,24 +1077,28 @@ void FMediaPlateCustomization::SetIsMediaSourceAsset(bool bIsAsset)
 	{
 		bIsMediaSourceAsset = bIsAsset;
 
-		// Clear out asset.
-		for (TWeakObjectPtr<UMediaPlateComponent>& MediaPlatePtr : MediaPlatesList)
 		{
-			UMediaPlateComponent* MediaPlate = MediaPlatePtr.Get();
-			if (MediaPlate != nullptr)
+			const FScopedTransaction Transaction(LOCTEXT("SetIsMediaSourceAsset", "Media Source Playlist Changed"));
+
+			// Clear out asset.
+			for (TWeakObjectPtr<UMediaPlateComponent>& MediaPlatePtr : MediaPlatesList)
 			{
-				// Get playlist.
-				UMediaPlaylist* Playlist = MediaPlate->MediaPlaylist;
-				if (Playlist != nullptr)
+				UMediaPlateComponent* MediaPlate = MediaPlatePtr.Get();
+				if (MediaPlate != nullptr)
 				{
-					// Update playlist.
-					if (Playlist->Num() > 0)
+					// Get playlist.
+					UMediaPlaylist* Playlist = MediaPlate->MediaPlaylist;
+					if (Playlist != nullptr)
 					{
-						Playlist->Replace(0, nullptr);
-						Playlist->MarkPackageDirty();
+						// Update playlist.
+						if (Playlist->Num() > 0)
+						{
+							Playlist->Modify();
+							Playlist->Replace(0, nullptr);
+						}
 					}
+					break;
 				}
-				break;
 			}
 		}
 
@@ -1167,14 +1171,19 @@ void FMediaPlateCustomization::OnPlaylistChanged(const FAssetData& AssetData)
 {
 	StopMediaPlates();
 
-	// Update the playlist.
-	UMediaPlaylist* Playlist = Cast<UMediaPlaylist>(AssetData.GetAsset());
-	for (TWeakObjectPtr<UMediaPlateComponent>& MediaPlatePtr : MediaPlatesList)
 	{
-		UMediaPlateComponent* MediaPlate = MediaPlatePtr.Get();
-		if (MediaPlate != nullptr)
+		const FScopedTransaction Transaction(LOCTEXT("OnPlaylistChanged", "Media Playlist Changed"));
+		
+		// Update the playlist.
+		UMediaPlaylist* Playlist = Cast<UMediaPlaylist>(AssetData.GetAsset());
+		for (TWeakObjectPtr<UMediaPlateComponent>& MediaPlatePtr : MediaPlatesList)
 		{
-			MediaPlate->MediaPlaylist = Playlist;
+			UMediaPlateComponent* MediaPlate = MediaPlatePtr.Get();
+			if (MediaPlate != nullptr)
+			{
+				MediaPlate->Modify();
+				MediaPlate->MediaPlaylist = Playlist;
+			}
 		}
 	}
 
@@ -1186,29 +1195,36 @@ void FMediaPlateCustomization::OnMediaSourceChanged(const FAssetData& AssetData)
 {
 	// Update the playlist with the new media source.
 	UMediaSource* MediaSource = Cast<UMediaSource>(AssetData.GetAsset());
-	for (TWeakObjectPtr<UMediaPlateComponent>& MediaPlatePtr : MediaPlatesList)
+
 	{
-		UMediaPlateComponent* MediaPlate = MediaPlatePtr.Get();
-		if (MediaPlate != nullptr)
+		const FScopedTransaction Transaction(LOCTEXT("OnMediaSourceChanged", "Media Source Playlist Changed"));
+
+		for (TWeakObjectPtr<UMediaPlateComponent>& MediaPlatePtr : MediaPlatesList)
 		{
-			// Get playlist.
-			UMediaPlaylist* Playlist = MediaPlate->MediaPlaylist;
-			if (Playlist == nullptr)
-			{
-				Playlist = NewObject<UMediaPlaylist>(MediaPlate);
-				MediaPlate->MediaPlaylist = Playlist;
-			}
-			
-			// Update playlist.
-			if (Playlist->Num() > 0)
-			{
-				Playlist->Replace(0, MediaSource);
-			}
-			else
+			UMediaPlateComponent* MediaPlate = MediaPlatePtr.Get();
+			if (MediaPlate != nullptr)
 			{
-				Playlist->Add(MediaSource);
+				// Get playlist.
+				TObjectPtr<UMediaPlaylist>& Playlist = MediaPlate->MediaPlaylist;
+				if (Playlist == nullptr)
+				{
+					Playlist = NewObject<UMediaPlaylist>(MediaPlate, NAME_None, RF_Transactional);
+				}
+				else
+				{
+					Playlist->Modify();
+				}
+
+				// Update playlist.
+				if (Playlist->Num() > 0)
+				{
+					Playlist->Replace(0, MediaSource);
+				}
+				else
+				{
+					Playlist->Add(MediaSource);
+				}
 			}
-			Playlist->MarkPackageDirty();
 		}
 	}
 
@@ -1265,33 +1281,39 @@ void FMediaPlateCustomization::HandleMediaPathPicked(const FString& PickedPath)
 		// Stop playback.
 		StopMediaPlates();
 
-		// Set up media source for our media plates.
-		for (TWeakObjectPtr<UMediaPlateComponent>& MediaPlatePtr : MediaPlatesList)
 		{
-			UMediaPlateComponent* MediaPlate = MediaPlatePtr.Get();
-			if (MediaPlate != nullptr)
+			const FScopedTransaction Transaction(LOCTEXT("OnMediaPathChanged", "Media Path Changed"));
+
+			// Set up media source for our media plates.
+			for (TWeakObjectPtr<UMediaPlateComponent>& MediaPlatePtr : MediaPlatesList)
 			{
-				// Get playlist.
-				UMediaPlaylist* Playlist = MediaPlate->MediaPlaylist;
-				if (Playlist == nullptr)
-				{
-					Playlist = NewObject<UMediaPlaylist>(MediaPlate);
-					MediaPlate->MediaPlaylist = Playlist;
-				}
-				
-				// Create media source for this path.
-				UMediaSource* MediaSource = UMediaSource::SpawnMediaSourceForString(PickedPath, MediaPlate);
-				if (MediaSource != nullptr)
+				UMediaPlateComponent* MediaPlate = MediaPlatePtr.Get();
+				if (MediaPlate != nullptr)
 				{
-					if (Playlist->Num() > 0)
+					// Get playlist.
+					TObjectPtr<UMediaPlaylist>& Playlist = MediaPlate->MediaPlaylist;
+					if (Playlist == nullptr)
 					{
-						Playlist->Replace(0, MediaSource);
+						Playlist = NewObject<UMediaPlaylist>(MediaPlate, NAME_None, RF_Transactional);
 					}
 					else
 					{
-						Playlist->Add(MediaSource);
+						Playlist->Modify();
+					}
+
+					// Create media source for this path.
+					UMediaSource* MediaSource = UMediaSource::SpawnMediaSourceForString(PickedPath, MediaPlate);
+					if (MediaSource != nullptr)
+					{
+						if (Playlist->Num() > 0)
+						{
+							Playlist->Replace(0, MediaSource);
+						}
+						else
+						{
+							Playlist->Add(MediaSource);
+						}
 					}
-					Playlist->MarkPackageDirty();
 				}
 			}
 		}
