MediaPlate: Supports layering for more than 2 tracks.

diff --git a/Engine/Plugins/Media/MediaPlate/Source/MediaPlate/Private/MediaPlateComponent.cpp b/Engine/Plugins/Media/MediaPlate/Source/MediaPlate/Private/MediaPlateComponent.cpp
--- a/Engine/Plugins/Media/MediaPlate/Source/MediaPlate/Private/MediaPlateComponent.cpp
+++ b/Engine/Plugins/Media/MediaPlate/Source/MediaPlate/Private/MediaPlateComponent.cpp
@@ -647,18 +647,18 @@ UMediaTexture* UMediaPlateComponent::ProxyGetMediaTexture(int32 LayerIndex, int3
 
 		// Fill up an empty slot if there is one.
 		bool bIsTextureSet = false;
-		for (int32 Index = 0; Index < TextureLayers[LayerIndex].Num(); ++Index)
+		for (int32 Index = 0; Index < TextureLayers[LayerIndex].Textures.Num(); ++Index)
 		{
-			if (TextureLayers[LayerIndex][Index] < 0)
+			if (TextureLayers[LayerIndex].Textures[Index] < 0)
 			{
-				TextureLayers[LayerIndex][Index] = TextureIndex;
+				TextureLayers[LayerIndex].Textures[Index] = TextureIndex;
 				bIsTextureSet = true;
 				break;
 			}
 		}
 		if (bIsTextureSet == false)
 		{
-			TextureLayers[LayerIndex].Add(TextureIndex);
+			TextureLayers[LayerIndex].Textures.Add(TextureIndex);
 		}
 
 		UpdateTextureLayers();
@@ -673,11 +673,11 @@ void UMediaPlateComponent::ProxyReleaseMediaTexture(int32 LayerIndex, int32 Text
 
 	if (LayerIndex < TextureLayers.Num())
 	{
-		for (int32 Index = 0; Index < TextureLayers[LayerIndex].Num(); ++Index)
+		for (int32 Index = 0; Index < TextureLayers[LayerIndex].Textures.Num(); ++Index)
 		{
-			if (TextureLayers[LayerIndex][Index] == TextureIndex)
+			if (TextureLayers[LayerIndex].Textures[Index] == TextureIndex)
 			{
-				TextureLayers[LayerIndex][Index] = -1;
+				TextureLayers[LayerIndex].Textures[Index] = -1;
 				break;
 			}
 		}
@@ -724,21 +724,26 @@ void UMediaPlateComponent::ProxySetTextureBlend(int32 LayerIndex, int32 TextureI
 			if (MID != nullptr)
 			{
 				int32 MatNumLayers = MediaTextures.Num() / MatNumTexPerLayer;
-				if ((LayerIndex < MatNumLayers) && (LayerIndex < TextureLayers.Num()))
+				if (LayerIndex < TextureLayers.Num())
 				{
-					const TArray<int32>& Layer = TextureLayers[LayerIndex];
-					for (int32 LayerTexIndex = 0;
-						(LayerTexIndex < MatNumTexPerLayer) && (LayerTexIndex < Layer.Num());
-						LayerTexIndex++)
+					int32 MaterialLayerIndex = TextureLayers[LayerIndex].MaterialLayerIndex;
+					if (MaterialLayerIndex < MatNumLayers)
 					{
-						if (Layer[LayerTexIndex] == TextureIndex)
+						const TArray<int32>& Layer = TextureLayers[LayerIndex].Textures;
+						for (int32 LayerTexIndex = 0;
+							(LayerTexIndex < MatNumTexPerLayer) && (LayerTexIndex < Layer.Num());
+							LayerTexIndex++)
 						{
-							int32 MatTexIndex = LayerIndex * MatNumTexPerLayer + LayerTexIndex;
-							static const FString BaseBlendName = TEXT("Blend");
-							FString BlendName = BaseBlendName;
-							BlendName.AppendInt(MatTexIndex);
-							MID->SetScalarParameterValue(FName(*BlendName), Blend);
-							break;
+							if (Layer[LayerTexIndex] == TextureIndex)
+							{
+								int32 MatTexIndex = MaterialLayerIndex * MatNumTexPerLayer
+									+ LayerTexIndex;
+								static const FString BaseBlendName = TEXT("Blend");
+								FString BlendName = BaseBlendName;
+								BlendName.AppendInt(MatTexIndex);
+								MID->SetScalarParameterValue(FName(*BlendName), Blend);
+								break;
+							}
 						}
 					}
 				}
@@ -1018,11 +1023,11 @@ void UMediaPlateComponent::SetNormalMode(bool bInIsNormalMode)
 			{
 				TextureLayers.SetNum(1);
 			}
-			if (TextureLayers[0].Num() != 1)
+			if (TextureLayers[0].Textures.Num() != 1)
 			{
-				TextureLayers[0].SetNum(1);
+				TextureLayers[0].Textures.SetNum(1);
 			}
-			TextureLayers[0][0] = 0;
+			TextureLayers[0].Textures[0] = 0;
 			UpdateTextureLayers();
 			
 			ProxySetTextureBlend(0, 0, 1.0f);
@@ -1049,19 +1054,25 @@ void UMediaPlateComponent::UpdateTextureLayers()
 				// Go through each layer.
 				int32 MatNumLayers = MediaTextures.Num() / MatNumTexPerLayer;
 				static const FString BaseTextureName = TEXT("MediaTexture");
-				int32 NumLayers = FMath::Min(MatNumLayers, TextureLayers.Num());
+				int32 NumLayers = TextureLayers.Num();
+				int32 MaterialLayerIndex = 0;
 				for (int32 LayerIndex = 0; LayerIndex < NumLayers; ++LayerIndex)
 				{
 					// Go through each texture in the layer.
-					const TArray<int32>& Layer = TextureLayers[LayerIndex];
+					const TArray<int32>& Layer = TextureLayers[LayerIndex].Textures;
 					int32 NumTex = FMath::Min(MatNumTexPerLayer, Layer.Num());
+					TextureLayers[LayerIndex].MaterialLayerIndex = -1;
 					for (int32 LayerTexIndex = 0; LayerTexIndex < NumTex; LayerTexIndex++)
 					{
 						// Set the texture in the material according to the layer data.
 						int32 TextureIndex = Layer[LayerTexIndex];
 						if (TextureIndex >= 0)
 						{
-							int32 MatTexIndex = LayerIndex * MatNumTexPerLayer + LayerTexIndex;
+							// Assign the next layer in the material to this layer.
+							TextureLayers[LayerIndex].MaterialLayerIndex = MaterialLayerIndex;
+							
+							int32 MatTexIndex = MaterialLayerIndex * MatNumTexPerLayer +
+								LayerTexIndex;
 							FString TextureName = BaseTextureName;
 							if (MatTexIndex != 0)
 							{
@@ -1070,6 +1081,17 @@ void UMediaPlateComponent::UpdateTextureLayers()
 							MID->SetTextureParameterValue(FName(*TextureName), MediaTextures[TextureIndex]);
 						}
 					}
+
+					// Did we use this layer?
+					if (TextureLayers[LayerIndex].MaterialLayerIndex != -1)
+					{
+						MaterialLayerIndex++;
+						// Did we run out of layers in the material?
+						if (MaterialLayerIndex >= MatNumLayers)
+						{
+							break;
+						}
+					}
 				}
 			}
 		}
diff --git a/Engine/Plugins/Media/MediaPlate/Source/MediaPlate/Public/MediaPlateComponent.h b/Engine/Plugins/Media/MediaPlate/Source/MediaPlate/Public/MediaPlateComponent.h
--- a/Engine/Plugins/Media/MediaPlate/Source/MediaPlate/Public/MediaPlateComponent.h
+++ b/Engine/Plugins/Media/MediaPlate/Source/MediaPlate/Public/MediaPlateComponent.h
@@ -368,7 +368,16 @@ private:
 	 * int32 is an index into MediaTextures.
 	 * -1 signifies no entry.
 	 */
-	TArray<TArray<int32>> TextureLayers;
+	struct Layer
+	{
+		/** The layer in the material that this layer uses. */
+		int32 MaterialLayerIndex;
+		/** List of textures in this layer. */
+		TArray<int32> Textures;
+
+		Layer() : MaterialLayerIndex(0) {};
+	};
+	TArray<Layer> TextureLayers;
 
 	/**
 	 * Plays a media source.
