MediaPlate: Changed default material button to a button that provides all the media plate materials.

#rnx

diff --git a/Engine/Plugins/Media/MediaPlate/Source/MediaPlateEditor/Private/Widgets/SMediaPlateEditorMaterial.cpp b/Engine/Plugins/Media/MediaPlate/Source/MediaPlateEditor/Private/Widgets/SMediaPlateEditorMaterial.cpp
--- a/Engine/Plugins/Media/MediaPlate/Source/MediaPlateEditor/Private/Widgets/SMediaPlateEditorMaterial.cpp
+++ b/Engine/Plugins/Media/MediaPlate/Source/MediaPlateEditor/Private/Widgets/SMediaPlateEditorMaterial.cpp
@@ -2,9 +2,15 @@
 
 #include "Widgets/SMediaPlateEditorMaterial.h"
 
+#include "AssetRegistry/AssetData.h"
+#include "AssetRegistry/AssetRegistryModule.h"
+#include "Framework/Commands/UIAction.h"
+#include "Framework/MultiBox/MultiBoxBuilder.h"
+#include "Materials/Material.h"
 #include "MediaPlate.h"
-#include "Widgets/Input/SButton.h"
+#include "ScopedTransaction.h"
 #include "Widgets/Layout/SBox.h"
+#include "Widgets/Input/SComboButton.h"
 
 #define LOCTEXT_NAMESPACE "SMediaPlateEditorMaterial"
 
@@ -26,10 +32,15 @@ void SMediaPlateEditorMaterial::Construct(const FArguments& InArgs, const TShare
 				.AutoHeight()
 				.VAlign(VAlign_Center)
 				[
-					SNew(SButton)
-						.Text(LOCTEXT("ApplyDefaultMaterial", "Use Media Plate Default"))
-						.ToolTipText(LOCTEXT("ApplyDefaultMaterialTooltip", "Set the material to the default Media Plate material."))
-						.OnClicked(this, &SMediaPlateEditorMaterial::OnDefaultButtonClicked)
+					SNew(SComboButton)
+						.OnGetMenuContent(this, &SMediaPlateEditorMaterial::OnGetMaterials)
+						.ContentPadding(2.0f)
+						.ButtonContent()
+						[
+							SNew(STextBlock)
+								.Text(LOCTEXT("SelectMaterial", "Select Media Plate Material"))
+								.ToolTipText(LOCTEXT("SelectMaterialTooltip", "Select a material to use from the recommended Media Plate materials."))
+						]
 				];
 		}
 	}
@@ -40,14 +51,42 @@ void SMediaPlateEditorMaterial::Construct(const FArguments& InArgs, const TShare
 	];
 }
 
-FReply SMediaPlateEditorMaterial::OnDefaultButtonClicked() const
+TSharedRef<SWidget> SMediaPlateEditorMaterial::OnGetMaterials()
 {
-	if (MediaPlate != nullptr)
+	FMenuBuilder MenuBuilder(true, NULL);
+
+	// Get media plate assets.
+	IAssetRegistry& AssetRegistry = FModuleManager::GetModuleChecked<FAssetRegistryModule>("AssetRegistry").Get();
+	TArray<FAssetData> Assets;
+	AssetRegistry.GetAssetsByPath(FName("/MediaPlate"), Assets, true);
+
+	for (const FAssetData& AssetData : Assets)
 	{
-		MediaPlate->UseDefaultMaterial();
+		if (AssetData.IsInstanceOf(UMaterial::StaticClass()))
+		{
+			FUIAction Action = FUIAction(FExecuteAction::CreateSP(this, &SMediaPlateEditorMaterial::OnSelectMaterial, AssetData));
+			MenuBuilder.AddMenuEntry(FText::FromName(AssetData.AssetName), FText(), FSlateIcon(), Action);
+		}
 	}
 
-	return FReply::Handled();
+	return MenuBuilder.MakeWidget();
+}
+
+void SMediaPlateEditorMaterial::OnSelectMaterial(FAssetData AssetData) const
+{
+	if (MediaPlate != nullptr)
+	{
+		UObject* AssetObject = AssetData.GetAsset();
+		if (AssetObject != nullptr)
+		{
+			UMaterial* Material = Cast<UMaterial>(AssetObject);
+			if (Material != nullptr)
+			{
+				const FScopedTransaction Transaction(LOCTEXT("SelectMaterial", "Media Plate Select Material"));
+				MediaPlate->ApplyMaterial(Material);
+			}
+		}
+	}
 }
 
 #undef LOCTEXT_NAMESPACE
diff --git a/Engine/Plugins/Media/MediaPlate/Source/MediaPlateEditor/Private/Widgets/SMediaPlateEditorMaterial.h b/Engine/Plugins/Media/MediaPlate/Source/MediaPlateEditor/Private/Widgets/SMediaPlateEditorMaterial.h
--- a/Engine/Plugins/Media/MediaPlate/Source/MediaPlateEditor/Private/Widgets/SMediaPlateEditorMaterial.h
+++ b/Engine/Plugins/Media/MediaPlate/Source/MediaPlateEditor/Private/Widgets/SMediaPlateEditorMaterial.h
@@ -30,9 +30,14 @@ public:
 
 private:
 	/**
-	 * Callback for when we press the default button.
+	 * Callback to get a list of media plate materials.
 	 */
-	FReply OnDefaultButtonClicked() const;
+	TSharedRef<SWidget> OnGetMaterials();
+
+	/**
+	 * Callback for when we select a material.
+	 */
+	void OnSelectMaterial(FAssetData AssetData) const;
 
 	/** Weak pointer to current component using the current material */
 	TWeakObjectPtr<AMediaPlate> MediaPlate = nullptr;
