MediaPlate: Fixed localization issue.
#rnx

diff --git a/Engine/Plugins/Media/MediaPlate/Source/MediaPlateEditor/Private/Widgets/SMediaPlateEditorMaterial.cpp b/Engine/Plugins/Media/MediaPlate/Source/MediaPlateEditor/Private/Widgets/SMediaPlateEditorMaterial.cpp
--- a/Engine/Plugins/Media/MediaPlate/Source/MediaPlateEditor/Private/Widgets/SMediaPlateEditorMaterial.cpp
+++ b/Engine/Plugins/Media/MediaPlate/Source/MediaPlateEditor/Private/Widgets/SMediaPlateEditorMaterial.cpp
@@ -38,7 +38,7 @@ void SMediaPlateEditorMaterial::Construct(const FArguments& InArgs, const TShare
 						.ButtonContent()
 						[
 							SNew(STextBlock)
-								.Text(LOCTEXT("SelectMaterial", "Select Media Plate Material"))
+								.Text(LOCTEXT("SelectMaterialButton", "Select Media Plate Material"))
 								.ToolTipText(LOCTEXT("SelectMaterialTooltip", "Select a material to use from the recommended Media Plate materials."))
 						]
 				];
@@ -82,7 +82,7 @@ void SMediaPlateEditorMaterial::OnSelectMaterial(FAssetData AssetData) const
 			UMaterial* Material = Cast<UMaterial>(AssetObject);
 			if (Material != nullptr)
 			{
-				const FScopedTransaction Transaction(LOCTEXT("SelectMaterial", "Media Plate Select Material"));
+				const FScopedTransaction Transaction(LOCTEXT("SelectMaterialTransaction", "Media Plate Select Material"));
 				MediaPlate->ApplyMaterial(Material);
 			}
 		}
