MediaPlate: Fixed missing icon.
#rnx

diff --git a/Engine/Plugins/Media/MediaPlate/Source/MediaPlateEditor/MediaPlateEditor.Build.cs b/Engine/Plugins/Media/MediaPlate/Source/MediaPlateEditor/MediaPlateEditor.Build.cs
--- a/Engine/Plugins/Media/MediaPlate/Source/MediaPlateEditor/MediaPlateEditor.Build.cs
+++ b/Engine/Plugins/Media/MediaPlate/Source/MediaPlateEditor/MediaPlateEditor.Build.cs
@@ -34,6 +34,7 @@ public class MediaPlateEditor : ModuleRules
 				"ModelingComponentsEditorOnly",
 				"ModelingComponents",
 				"MovieScene",
+				"Projects",
 				"PropertyEditor",
 				"SceneOutliner",
 				"Sequencer",
diff --git a/Engine/Plugins/Media/MediaPlate/Source/MediaPlateEditor/Private/MediaPlateEditorStyle.cpp b/Engine/Plugins/Media/MediaPlate/Source/MediaPlateEditor/Private/MediaPlateEditorStyle.cpp
--- a/Engine/Plugins/Media/MediaPlate/Source/MediaPlateEditor/Private/MediaPlateEditorStyle.cpp
+++ b/Engine/Plugins/Media/MediaPlate/Source/MediaPlateEditor/Private/MediaPlateEditorStyle.cpp
@@ -6,6 +6,7 @@
 #include "Brushes/SlateBoxBrush.h"
 #include "Brushes/SlateImageBrush.h"
 #include "Fonts/SlateFontInfo.h"
+#include "Interfaces/IPluginManager.h"
 #include "Math/Vector2D.h"
 #include "Misc/Paths.h"
 #include "Styling/SlateStyleRegistry.h"
@@ -13,6 +14,7 @@
 #include "Styling/CoreStyle.h"
 
 #define IMAGE_BRUSH(RelativePath, ...) FSlateImageBrush(RootToContentDir(RelativePath, TEXT(".png")), __VA_ARGS__)
+#define IMAGE_BRUSH_SVG(RelativePath, ...) FSlateVectorImageBrush(RootToContentDir(RelativePath, TEXT(".svg" )), __VA_ARGS__)
 #define BOX_BRUSH(RelativePath, ...) FSlateBoxBrush(RootToContentDir(RelativePath, TEXT(".png")), __VA_ARGS__)
 #define BORDER_BRUSH(RelativePath, ...) FSlateBorderBrush(RootToContentDir(RelativePath, TEXT(".png")), __VA_ARGS__)
 #define DEFAULT_FONT(...) FCoreStyle::GetDefaultFontStyle(__VA_ARGS__)
@@ -28,6 +30,17 @@ FMediaPlateEditorStyle::FMediaPlateEditorStyle()
 	const FVector2D Icon16x16(16.0f, 16.0f);
 	const FVector2D Icon20x20(20.0f, 20.0f);
 	const FVector2D Icon40x40(40.0f, 40.0f);
+	const FVector2D Icon64x64(64.0f, 64.0f);
+
+	if (TSharedPtr<IPlugin> Plugin = IPluginManager::Get().FindPlugin(TEXT("MediaPlate")))
+	{
+		SetContentRoot(FPaths::Combine(Plugin->GetBaseDir(), TEXT("Resources")));
+	}
+	// Class icons.
+	Set("ClassIcon.MediaPlate", new IMAGE_BRUSH_SVG(TEXT("Icons/MediaPlate"), Icon16x16));
+	Set("ClassThumbnail.MediaPlate", new IMAGE_BRUSH_SVG(TEXT("Icons/MediaPlate"), Icon64x64));
+	Set("ClassIcon.MediaPlateComponent", new IMAGE_BRUSH_SVG(TEXT("Icons/MediaPlate"), Icon16x16));
+	Set("ClassThumbnail.MediaPlateComponent", new IMAGE_BRUSH_SVG(TEXT("Icons/MediaPlate"), Icon64x64));
 
 	SetContentRoot(FPaths::EnginePluginsDir() / TEXT("Media/MediaPlayerEditor/Content"));
 
