Media Plate Component: Added event that gets broadcast when any of custom UI buttons are pressed. This will be used to transact changes to remote endpoints that is addressed in another CL.

#rnx

diff --git a/Engine/Plugins/Media/MediaPlate/Source/MediaPlate/Private/MediaPlateComponent.cpp b/Engine/Plugins/Media/MediaPlate/Source/MediaPlate/Private/MediaPlateComponent.cpp
--- a/Engine/Plugins/Media/MediaPlate/Source/MediaPlate/Private/MediaPlateComponent.cpp
+++ b/Engine/Plugins/Media/MediaPlate/Source/MediaPlate/Private/MediaPlateComponent.cpp
@@ -752,6 +752,32 @@ void UMediaPlateComponent::ProxySetTextureBlend(int32 LayerIndex, int32 TextureI
 	}
 }
 
+#if WITH_EDITOR
+float UMediaPlateComponent::GetForwardRate(UMediaPlayer* MediaPlayer)
+{
+	float Rate = MediaPlayer->GetRate();
+
+	if (Rate < 1.0f)
+	{
+		Rate = 1.0f;
+	}
+
+	return 2.0f * Rate;
+}
+
+float UMediaPlateComponent::GetReverseRate(UMediaPlayer* MediaPlayer)
+{
+	float Rate = MediaPlayer->GetRate();
+
+	if (Rate > -1.0f)
+	{
+		return -1.0f;
+	}
+
+	return 2.0f * Rate;
+}
+#endif
+
 void UMediaPlateComponent::RestartPlayer()
 {
 	if (MediaPlayer != nullptr)
@@ -1186,6 +1212,60 @@ void UMediaPlateComponent::PostEditChangeProperty(FPropertyChangedEvent& Propert
 	}
 }
 
+void UMediaPlateComponent::SwitchStates(EMediaPlateEventState State)
+{
+	switch (State)
+	{
+	case EMediaPlateEventState::Play:
+		{
+			Play();
+		}
+		break;
+	case EMediaPlateEventState::Open:
+		{
+			Open();
+		}
+		break;
+	case EMediaPlateEventState::Close:
+		{
+			Close();
+		}
+		break;
+	case EMediaPlateEventState::Pause:
+		{
+			Pause();
+		}
+		break;
+	case EMediaPlateEventState::Reverse:
+		{
+			if (MediaPlayer != nullptr)
+			{
+				MediaPlayer->SetRate(GetReverseRate(MediaPlayer));
+			}
+		}
+		break;
+	case EMediaPlateEventState::Forward:
+		{
+			if (MediaPlayer != nullptr)
+			{
+				MediaPlayer->SetRate(GetForwardRate(MediaPlayer));
+			}
+		}
+		break;
+	case EMediaPlateEventState::Rewind:
+		{
+			if (MediaPlayer != nullptr)
+			{
+				MediaPlayer->Rewind();
+			}
+		}
+		break;
+	case EMediaPlateEventState::MAX:
+	default:
+		checkNoEntry();
+		break;
+	}
+}
 #endif // WITH_EDITOR
 
 #undef LOCTEXT_NAMESPACE
diff --git a/Engine/Plugins/Media/MediaPlate/Source/MediaPlate/Private/MediaPlateModule.cpp b/Engine/Plugins/Media/MediaPlate/Source/MediaPlate/Private/MediaPlateModule.cpp
--- a/Engine/Plugins/Media/MediaPlate/Source/MediaPlate/Private/MediaPlateModule.cpp
+++ b/Engine/Plugins/Media/MediaPlate/Source/MediaPlate/Private/MediaPlateModule.cpp
@@ -33,7 +33,12 @@ void FMediaPlateModule::StartupModule()
 	if (MediaAssetsModule != nullptr)
 	{
 		GetPlayerFromObjectID = MediaAssetsModule->RegisterGetPlayerFromObject(IMediaAssetsModule::FOnGetPlayerFromObject::CreateRaw(this, &FMediaPlateModule::GetMediaPlayer));
+
+#if WITH_EDITOR
+		OnMediaPlateStateChangedHandle = MediaAssetsModule->RegisterOnMediaStateChangedEvent(IMediaAssetsModule::FMediaStateChangedDelegate::FDelegate::CreateRaw(this, &FMediaPlateModule::OnMediaPlateStateChanged));
+#endif
 	}
+
 }
 
 void FMediaPlateModule::ShutdownModule()
@@ -45,8 +50,36 @@ void FMediaPlateModule::ShutdownModule()
 		{
 			MediaAssetsModule->UnregisterGetPlayerFromObject(GetPlayerFromObjectID);
 			GetPlayerFromObjectID = INDEX_NONE;
+
+#if WITH_EDITOR
+			MediaAssetsModule->UnregisterOnMediaStateChangedEvent(OnMediaPlateStateChangedHandle);
+#endif
+		}
+	}
+}
+
+#if WITH_EDITOR
+void FMediaPlateModule::OnMediaPlateStateChanged(const TArray<FString>& InActorsPathNames, uint8 InEnumState, bool bRemoteBroadcast)
+{
+	// Button was pressed locally. This event will be handled in FMediaPlateCustomization.
+	if (!bRemoteBroadcast)
+	{
+		return;
+	}
+	for (const FString& ActorPathName : InActorsPathNames)
+	{
+		AActor* MediaPlateActor = Cast<AActor>(StaticFindObject(AActor::StaticClass(), nullptr, *ActorPathName, false));
+		if (MediaPlateActor)
+		{
+			TArray<UMediaPlateComponent*> MediaPlateComponents;
+			MediaPlateActor->GetComponents<UMediaPlateComponent>(MediaPlateComponents);
+			for (UMediaPlateComponent* MediaPlateComponent : MediaPlateComponents)
+			{
+				MediaPlateComponent->SwitchStates((EMediaPlateEventState)InEnumState);
+			}
 		}
 	}
 }
+#endif
 
 IMPLEMENT_MODULE(FMediaPlateModule, MediaPlate)
diff --git a/Engine/Plugins/Media/MediaPlate/Source/MediaPlate/Public/MediaPlateComponent.h b/Engine/Plugins/Media/MediaPlate/Source/MediaPlate/Public/MediaPlateComponent.h
--- a/Engine/Plugins/Media/MediaPlate/Source/MediaPlate/Public/MediaPlateComponent.h
+++ b/Engine/Plugins/Media/MediaPlate/Source/MediaPlate/Public/MediaPlateComponent.h
@@ -25,6 +25,19 @@ namespace UE::MediaPlateComponent
 	enum class ESetUpTexturesFlags;
 }
 
+UENUM()
+enum class EMediaPlateEventState : uint8
+{
+	Play,
+	Open,
+	Close,
+	Pause,
+	Reverse,
+	Forward,
+	Rewind,
+	MAX
+};
+
 /**
  * This is a component for AMediaPlate that can play and show media in the world.
  */
@@ -242,6 +255,11 @@ public:
 	 * Call this to set the mip tile calculations mesh mode. (Note: restarts playback to apply changes.)
 	 */
 	void SetVisibleMipsTilesCalculations(EMediaTextureVisibleMipsTiles InVisibleMipsTilesCalculations);
+
+	/**
+	* Called whenever a button was pressed locally or on a remote endpoint.
+	*/
+	void SwitchStates(EMediaPlateEventState State);
 #endif
 
 	/**
@@ -260,6 +278,19 @@ public:
 	virtual bool ProxySetAspectRatio(UMediaPlayer* InMediaPlayer) override;
 	virtual void ProxySetTextureBlend(int32 LayerIndex, int32 TextureIndex, float Blend) override;
 
+#if WITH_EDITOR
+public:
+	/**
+	 * Get the rate to use when we press the forward button.
+	 */
+	static float GetForwardRate(UMediaPlayer* MediaPlayer);
+
+	/**
+	 * Get the rate to use when we press the reverse button.
+	 */
+	static float GetReverseRate(UMediaPlayer* MediaPlayer);
+#endif
+
 private:
 	/**
 	 * Adds our media texture to the media texture tracker.
diff --git a/Engine/Plugins/Media/MediaPlate/Source/MediaPlate/Public/MediaPlateModule.h b/Engine/Plugins/Media/MediaPlate/Source/MediaPlate/Public/MediaPlateModule.h
--- a/Engine/Plugins/Media/MediaPlate/Source/MediaPlate/Public/MediaPlateModule.h
+++ b/Engine/Plugins/Media/MediaPlate/Source/MediaPlate/Public/MediaPlateModule.h
@@ -37,7 +37,18 @@ public:
 	/** Add to this if you want a callback when a media plate applies a material to itself. */
 	FOnMediaPlateApplyMaterial OnMediaPlateApplyMaterial;
 
+#if WITH_EDITOR
+	/**
+	* Called when Media playback state changed.
+	*/
+	void OnMediaPlateStateChanged(const TArray<FString>& InActorsPathNames, uint8 InEnumState, bool bRemoteBroadcast);
+#endif
 private:
 	/** ID for our delegate. */
 	int32 GetPlayerFromObjectID = INDEX_NONE;
+
+#if WITH_EDITOR
+	/** Handle for the subscribed event */
+	FDelegateHandle OnMediaPlateStateChangedHandle;
+#endif
 };
diff --git a/Engine/Plugins/Media/MediaPlate/Source/MediaPlateEditor/Private/MediaPlateCustomization.cpp b/Engine/Plugins/Media/MediaPlate/Source/MediaPlateEditor/Private/MediaPlateCustomization.cpp
--- a/Engine/Plugins/Media/MediaPlate/Source/MediaPlateEditor/Private/MediaPlateCustomization.cpp
+++ b/Engine/Plugins/Media/MediaPlate/Source/MediaPlateEditor/Private/MediaPlateCustomization.cpp
@@ -9,6 +9,7 @@
 #include "DetailWidgetRow.h"
 #include "Framework/MultiBox/MultiBoxBuilder.h"
 #include "IDetailGroup.h"
+#include "IMediaAssetsModule.h"
 #include "LevelEditorViewport.h"
 #include "MediaPlate.h"
 #include "MediaPlateComponent.h"
@@ -218,18 +219,7 @@ void FMediaPlateCustomization::CustomizeDetails(IDetailLayoutBuilder& DetailBuil
 						})
 						.OnClicked_Lambda([this]() -> FReply
 						{
-							for (TWeakObjectPtr<UMediaPlateComponent>& MediaPlatePtr : MediaPlatesList)
-							{
-								UMediaPlateComponent* MediaPlate = MediaPlatePtr.Get();
-								if (MediaPlate != nullptr)
-								{
-									UMediaPlayer* MediaPlayer = MediaPlate->GetMediaPlayer();
-									if (MediaPlayer != nullptr)
-									{
-										MediaPlayer->Rewind();
-									}
-								}
-							}
+							OnButtonEvent(EMediaPlateEventState::Rewind);
 							return FReply::Handled();
 						})
 						[
@@ -256,7 +246,7 @@ void FMediaPlateCustomization::CustomizeDetails(IDetailLayoutBuilder& DetailBuil
 									UMediaPlayer* MediaPlayer = MediaPlate->GetMediaPlayer();
 									if (MediaPlayer != nullptr)
 									{
-										return MediaPlayer->IsReady() && MediaPlayer->SupportsRate(GetReverseRate(MediaPlayer), false);
+										return MediaPlayer->IsReady() && MediaPlayer->SupportsRate(UMediaPlateComponent::GetReverseRate(MediaPlayer), false);
 									}
 								}
 							}
@@ -264,18 +254,7 @@ void FMediaPlateCustomization::CustomizeDetails(IDetailLayoutBuilder& DetailBuil
 						})
 						.OnClicked_Lambda([this]() -> FReply
 						{
-							for (TWeakObjectPtr<UMediaPlateComponent>& MediaPlatePtr : MediaPlatesList)
-							{
-								UMediaPlateComponent* MediaPlate = MediaPlatePtr.Get();
-								if (MediaPlate != nullptr)
-								{
-									UMediaPlayer* MediaPlayer = MediaPlate->GetMediaPlayer();
-									if (MediaPlayer != nullptr)
-									{
-										MediaPlayer->SetRate(GetReverseRate(MediaPlayer));
-									}
-								}
-							}
+							OnButtonEvent(EMediaPlateEventState::Reverse);
 							return FReply::Handled();
 						})
 						[
@@ -311,14 +290,7 @@ void FMediaPlateCustomization::CustomizeDetails(IDetailLayoutBuilder& DetailBuil
 						})
 						.OnClicked_Lambda([this]() -> FReply
 						{
-							for (TWeakObjectPtr<UMediaPlateComponent>& MediaPlatePtr : MediaPlatesList)
-							{
-								UMediaPlateComponent* MediaPlate = MediaPlatePtr.Get();
-								if (MediaPlate != nullptr)
-								{
-									MediaPlate->Play();
-								}
-							}
+							OnButtonEvent(EMediaPlateEventState::Play);
 							return FReply::Handled();
 						})
 						[
@@ -354,14 +326,7 @@ void FMediaPlateCustomization::CustomizeDetails(IDetailLayoutBuilder& DetailBuil
 						})
 						.OnClicked_Lambda([this]() -> FReply
 						{
-							for (TWeakObjectPtr<UMediaPlateComponent>& MediaPlatePtr : MediaPlatesList)
-							{
-								UMediaPlateComponent* MediaPlate = MediaPlatePtr.Get();
-								if (MediaPlate != nullptr)
-								{							
-									MediaPlate->Pause();
-								}
-							}
+							OnButtonEvent(EMediaPlateEventState::Pause);
 							return FReply::Handled();
 						})
 						[
@@ -388,7 +353,7 @@ void FMediaPlateCustomization::CustomizeDetails(IDetailLayoutBuilder& DetailBuil
 								UMediaPlayer* MediaPlayer = MediaPlate->GetMediaPlayer();
 								if (MediaPlayer != nullptr)
 								{
-									return MediaPlayer->IsReady() && MediaPlayer->SupportsRate(GetForwardRate(MediaPlayer), false);
+									return MediaPlayer->IsReady() && MediaPlayer->SupportsRate(UMediaPlateComponent::GetForwardRate(MediaPlayer), false);
 								}
 							}
 						}
@@ -396,18 +361,7 @@ void FMediaPlateCustomization::CustomizeDetails(IDetailLayoutBuilder& DetailBuil
 					})
 					.OnClicked_Lambda([this]() -> FReply
 					{
-						for (TWeakObjectPtr<UMediaPlateComponent>& MediaPlatePtr : MediaPlatesList)
-						{
-							UMediaPlateComponent* MediaPlate = MediaPlatePtr.Get();
-							if (MediaPlate != nullptr)
-							{
-								UMediaPlayer* MediaPlayer = MediaPlate->GetMediaPlayer();
-								if (MediaPlayer != nullptr)
-								{
-									MediaPlayer->SetRate(GetForwardRate(MediaPlayer));
-								}
-							}
-						}
+						OnButtonEvent(EMediaPlateEventState::Forward);
 						return FReply::Handled();
 					})
 					[
@@ -427,22 +381,7 @@ void FMediaPlateCustomization::CustomizeDetails(IDetailLayoutBuilder& DetailBuil
 						.VAlign(VAlign_Center)
 						.OnClicked_Lambda([this]() -> FReply
 						{
-							for (TWeakObjectPtr<UMediaPlateComponent>& MediaPlatePtr : MediaPlatesList)
-							{
-								UMediaPlateComponent* MediaPlate = MediaPlatePtr.Get();
-								if (MediaPlate != nullptr)
-								{
-									// Tell the editor module that this media plate is playing.
-									FMediaPlateEditorModule* EditorModule = FModuleManager::LoadModulePtr<FMediaPlateEditorModule>("MediaPlateEditor");
-									if (EditorModule != nullptr)
-									{
-										EditorModule->MediaPlateStartedPlayback(MediaPlate);
-									}
-
-									// Play the media.
-									MediaPlate->Open();
-								}
-							}
+							OnButtonEvent(EMediaPlateEventState::Open);
 							return FReply::Handled();
 						})
 						[
@@ -478,14 +417,7 @@ void FMediaPlateCustomization::CustomizeDetails(IDetailLayoutBuilder& DetailBuil
 						})
 						.OnClicked_Lambda([this]() -> FReply
 						{
-							for (TWeakObjectPtr<UMediaPlateComponent>& MediaPlatePtr : MediaPlatesList)
-							{
-								UMediaPlateComponent* MediaPlate = MediaPlatePtr.Get();
-								if (MediaPlate != nullptr)
-								{
-									MediaPlate->Close();
-								}
-							}
+							OnButtonEvent(EMediaPlateEventState::Close);
 							return FReply::Handled();
 						})
 						[
@@ -1232,6 +1164,35 @@ void FMediaPlateCustomization::OnMediaSourceChanged(const FAssetData& AssetData)
 	UpdateMediaPath();
 }
 
+void FMediaPlateCustomization::OnButtonEvent(EMediaPlateEventState State)
+{
+	IMediaAssetsModule* MediaAssets = FModuleManager::LoadModulePtr<IMediaAssetsModule>("MediaAssets");
+	TArray<FString> ActorsPathNames;
+	ActorsPathNames.Reserve(MediaPlatesList.Num());
+	for (TWeakObjectPtr<UMediaPlateComponent>& MediaPlatePtr : MediaPlatesList)
+	{
+		if (!MediaPlatePtr.IsValid())
+		{
+			continue;
+		}
+
+		ActorsPathNames.Add(MediaPlatePtr->GetOwner()->GetPathName());
+			
+		if (State == EMediaPlateEventState::Open)
+		{
+			// Tell the editor module that this media plate is playing.
+			FMediaPlateEditorModule* EditorModule = FModuleManager::LoadModulePtr<FMediaPlateEditorModule>("MediaPlateEditor");
+			if (EditorModule != nullptr)
+			{
+				EditorModule->MediaPlateStartedPlayback(MediaPlatePtr.Get());
+			}
+		}
+
+		MediaPlatePtr->SwitchStates(State);
+	}
+	MediaAssets->BroadcastOnMediaStateChangedEvent(ActorsPathNames, (uint8)State);
+}
+
 void FMediaPlateCustomization::UpdateMediaPath()
 {
 	MediaPath.Empty();
@@ -1347,38 +1308,7 @@ FReply FMediaPlateCustomization::OnOpenMediaPlate()
 
 void FMediaPlateCustomization::StopMediaPlates()
 {
-	for (TWeakObjectPtr<UMediaPlateComponent>& MediaPlatePtr : MediaPlatesList)
-	{
-		UMediaPlateComponent* MediaPlate = MediaPlatePtr.Get();
-		if (MediaPlate != nullptr)
-		{
-			MediaPlate->Close();
-		}
-	}
-}
-
-float FMediaPlateCustomization::GetForwardRate(UMediaPlayer* MediaPlayer) const
-{
-	float Rate = MediaPlayer->GetRate();
-
-	if (Rate < 1.0f)
-	{
-		Rate = 1.0f;
-	}
-
-	return 2.0f * Rate;
-}
-
-float FMediaPlateCustomization::GetReverseRate(UMediaPlayer* MediaPlayer) const
-{
-	float Rate = MediaPlayer->GetRate();
-
-	if (Rate > -1.0f)
-	{
-		return -1.0f;
-	}
-
-	return 2.0f * Rate;
+	OnButtonEvent(EMediaPlateEventState::Close);
 }
 
 
diff --git a/Engine/Plugins/Media/MediaPlate/Source/MediaPlateEditor/Private/MediaPlateCustomization.h b/Engine/Plugins/Media/MediaPlate/Source/MediaPlateEditor/Private/MediaPlateCustomization.h
--- a/Engine/Plugins/Media/MediaPlate/Source/MediaPlateEditor/Private/MediaPlateCustomization.h
+++ b/Engine/Plugins/Media/MediaPlate/Source/MediaPlateEditor/Private/MediaPlateCustomization.h
@@ -16,6 +16,7 @@ class FMenuBuilder;
 class SWidget;
 class UMediaPlateComponent;
 class UMediaPlayer;
+enum class EMediaPlateEventState : uint8;
 
 /**
  * Implements a details view customization for the UMediaPlateComponent class.
@@ -211,6 +212,11 @@ private:
 	 */
 	void OnMediaSourceChanged(const FAssetData& AssetData);
 
+	/**
+	* Changes the state of selected media plates and broadcasts the even to the remote endpoints.
+	*/
+	void OnButtonEvent(EMediaPlateEventState State);
+
 	/**
 	 * Updates MediaPath from the current MediaSource.
 	 */
@@ -235,15 +241,5 @@ private:
 	 * Call this to stop all playback.
 	 */
 	void StopMediaPlates();
-
-	/**
-	 * Get the rate to use when we press the forward button.
-	 */
-	float GetForwardRate(UMediaPlayer* MediaPlayer) const;
-
-	/**
-	 * Get the rate to use when we press the reverse button.
-	 */
-	float GetReverseRate(UMediaPlayer* MediaPlayer) const;
 };
 
diff --git a/Engine/Source/Runtime/MediaAssets/Private/MediaAssetsModule.cpp b/Engine/Source/Runtime/MediaAssets/Private/MediaAssetsModule.cpp
--- a/Engine/Source/Runtime/MediaAssets/Private/MediaAssetsModule.cpp
+++ b/Engine/Source/Runtime/MediaAssets/Private/MediaAssetsModule.cpp
@@ -102,6 +102,21 @@ public:
 		CreateMediaSourceRendererDelegate.Unbind();
 	}
 
+	virtual FDelegateHandle RegisterOnMediaStateChangedEvent(FMediaStateChangedDelegate::FDelegate InStateChangedDelegate) override 
+	{ 
+		return OnMediaStateChanged.Add(InStateChangedDelegate);
+	}
+
+	virtual void UnregisterOnMediaStateChangedEvent(FDelegateHandle InHandle) override 
+	{
+		OnMediaStateChanged.Remove(InHandle);
+	}
+
+	virtual void BroadcastOnMediaStateChangedEvent(const TArray<FString>& InActorsPathNames, uint8 EnumState, bool bRemoteBroadcast = false) override
+	{
+		OnMediaStateChanged.Broadcast(InActorsPathNames, EnumState, bRemoteBroadcast);
+	}
+
 	virtual UObject* CreateMediaSourceRenderer() override
 	{
 		UObject* Renderer = nullptr;
@@ -168,6 +183,9 @@ private:
 
 	/** Delegate to create an object that implements IMediaSourceRendererInterface. */
 	FOnCreateMediaSourceRenderer CreateMediaSourceRendererDelegate;
+
+	/** Delegate that reacts to change in Media state (Play, Stop, Pause etc.) */
+	FMediaStateChangedDelegate OnMediaStateChanged;
 };
 
 
diff --git a/Engine/Source/Runtime/MediaAssets/Public/IMediaAssetsModule.h b/Engine/Source/Runtime/MediaAssets/Public/IMediaAssetsModule.h
--- a/Engine/Source/Runtime/MediaAssets/Public/IMediaAssetsModule.h
+++ b/Engine/Source/Runtime/MediaAssets/Public/IMediaAssetsModule.h
@@ -20,6 +20,8 @@ public:
 	DECLARE_DELEGATE_RetVal_TwoParams(UMediaPlayer*, FOnGetPlayerFromObject, UObject*, UObject*& /*PlayerProxy*/);
 	/** Delegate to create an object that implements IMediaSourceRendererInterface. */
 	DECLARE_DELEGATE_RetVal(UObject*, FOnCreateMediaSourceRenderer);
+	/** Delegate that reacts to change in Media state (Play, Stop, Pause etc.) */
+	DECLARE_MULTICAST_DELEGATE_ThreeParams(FMediaStateChangedDelegate, const TArray<FString>& /*InActorsPathNames*/, uint8 /*EnumState*/, bool /*bRemoteBroadcast*/);
 
 	/**
 	 * Plugins should call this so they can provide a function to get a media player from an object.
@@ -58,6 +60,15 @@ public:
 	 */
 	virtual void UnregisterCreateMediaSourceRenderer() = 0;
 
+	/** Subscribes to the event that is called whenever any of Media state changes (such as play button was pressed). */
+	virtual FDelegateHandle RegisterOnMediaStateChangedEvent(FMediaStateChangedDelegate::FDelegate InStateChangedDelegate) = 0;
+
+	/** Removes event handling. */
+	virtual void UnregisterOnMediaStateChangedEvent(FDelegateHandle InHandle) = 0;
+
+	/* Broadcasts the event to all subscribers. InActorsPathNames are the paths of selected Medias in the scene. */
+	virtual void BroadcastOnMediaStateChangedEvent(const TArray<FString>& InActorsPathNames, uint8 EnumState, bool bRemoteBroadcast = false) = 0;
+
 	/**
 	 * Creates an object that implements IMediaSourceRendererInterface.
 	 */
