MediaPlate: Added button to browse to the material asset.

diff --git a/Engine/Plugins/Media/MediaPlate/Source/MediaPlateEditor/Private/Widgets/SMediaPlateEditorMaterial.cpp b/Engine/Plugins/Media/MediaPlate/Source/MediaPlateEditor/Private/Widgets/SMediaPlateEditorMaterial.cpp
--- a/Engine/Plugins/Media/MediaPlate/Source/MediaPlateEditor/Private/Widgets/SMediaPlateEditorMaterial.cpp
+++ b/Engine/Plugins/Media/MediaPlate/Source/MediaPlateEditor/Private/Widgets/SMediaPlateEditorMaterial.cpp
@@ -4,12 +4,15 @@
 
 #include "AssetRegistry/AssetData.h"
 #include "AssetRegistry/AssetRegistryModule.h"
+#include "Editor.h"
 #include "Framework/Commands/UIAction.h"
 #include "Framework/MultiBox/MultiBoxBuilder.h"
 #include "Materials/Material.h"
+#include "Materials/MaterialInstance.h"
 #include "MediaPlate.h"
 #include "ScopedTransaction.h"
 #include "Widgets/Layout/SBox.h"
+#include "Widgets/Input/SButton.h"
 #include "Widgets/Input/SComboButton.h"
 
 #define LOCTEXT_NAMESPACE "SMediaPlateEditorMaterial"
@@ -27,6 +30,17 @@ void SMediaPlateEditorMaterial::Construct(const FArguments& InArgs, const TShare
 		MediaPlate = InCurrentComponent->GetOwner<AMediaPlate>();
 		if (MediaPlate != nullptr)
 		{
+			// Add button to browse the material.
+			ResultWidget->AddSlot()
+				.AutoHeight()
+				.VAlign(VAlign_Center)
+				[
+					SNew(SButton)
+						.Text(LOCTEXT("BrowseMaterial", "Browse To Material"))
+						.ToolTipText(LOCTEXT("BrowsetMaterialTooltip", "Browse to the material asset in the Content Browser."))
+						.OnClicked(this, &SMediaPlateEditorMaterial::OnBrowseMaterial)
+				];
+
 			// Add button for default material.
 			ResultWidget->AddSlot()
 				.AutoHeight()
@@ -51,6 +65,26 @@ void SMediaPlateEditorMaterial::Construct(const FArguments& InArgs, const TShare
 	];
 }
 
+FReply SMediaPlateEditorMaterial::OnBrowseMaterial() const
+{
+	if (MediaPlate != nullptr)
+	{
+		UMaterialInterface* Material = MediaPlate->GetCurrentMaterial();
+		if (Material != nullptr)
+		{
+			UMaterial* BaseMaterial = Material->GetMaterial();
+			if (BaseMaterial != nullptr)
+			{
+				TArray<UObject*> ObjectsToSync;
+				ObjectsToSync.Add(BaseMaterial);
+				GEditor->SyncBrowserToObjects(ObjectsToSync);
+			}
+		}
+	}
+
+	return FReply::Handled();
+}
+
 TSharedRef<SWidget> SMediaPlateEditorMaterial::OnGetMaterials()
 {
 	FMenuBuilder MenuBuilder(true, NULL);
diff --git a/Engine/Plugins/Media/MediaPlate/Source/MediaPlateEditor/Private/Widgets/SMediaPlateEditorMaterial.h b/Engine/Plugins/Media/MediaPlate/Source/MediaPlateEditor/Private/Widgets/SMediaPlateEditorMaterial.h
--- a/Engine/Plugins/Media/MediaPlate/Source/MediaPlateEditor/Private/Widgets/SMediaPlateEditorMaterial.h
+++ b/Engine/Plugins/Media/MediaPlate/Source/MediaPlateEditor/Private/Widgets/SMediaPlateEditorMaterial.h
@@ -29,6 +29,11 @@ public:
 	void Construct(const FArguments& InArgs, const TSharedRef<FMaterialItemView>& InMaterialItemView, UActorComponent* InCurrentComponent);
 
 private:
+	/**
+	 * Callback for when we browse a material.
+	 */
+	FReply OnBrowseMaterial() const;
+
 	/**
 	 * Callback to get a list of media plate materials.
 	 */
