MediaPlate: Added placement category for media plate.
[lockdown] alejandro.arango
#rnx

Merging //UE5/Release-5.2 to Release-Engine-Staging (//UE5/Release-Engine-Staging)

diff --git a/Engine/Plugins/Media/MediaPlate/Source/MediaPlateEditor/MediaPlateEditor.Build.cs b/Engine/Plugins/Media/MediaPlate/Source/MediaPlateEditor/MediaPlateEditor.Build.cs
--- a/Engine/Plugins/Media/MediaPlate/Source/MediaPlateEditor/MediaPlateEditor.Build.cs
+++ b/Engine/Plugins/Media/MediaPlate/Source/MediaPlateEditor/MediaPlateEditor.Build.cs
@@ -34,6 +34,7 @@ public class MediaPlateEditor : ModuleRules
 				"ModelingComponentsEditorOnly",
 				"ModelingComponents",
 				"MovieScene",
+				"PlacementMode",
 				"Projects",
 				"PropertyEditor",
 				"SceneOutliner",
diff --git a/Engine/Plugins/Media/MediaPlate/Source/MediaPlateEditor/Private/MediaPlateEditorModule.cpp b/Engine/Plugins/Media/MediaPlate/Source/MediaPlateEditor/Private/MediaPlateEditorModule.cpp
--- a/Engine/Plugins/Media/MediaPlate/Source/MediaPlateEditor/Private/MediaPlateEditorModule.cpp
+++ b/Engine/Plugins/Media/MediaPlate/Source/MediaPlateEditor/Private/MediaPlateEditorModule.cpp
@@ -3,9 +3,11 @@
 #include "MediaPlateEditorModule.h"
 
 #include "Editor.h"
+#include "IPlacementModeModule.h"
 #include "ISequencerModule.h"
 #include "LevelEditor.h"
 #include "MaterialList.h"
+#include "MediaPlate.h"
 #include "MediaPlateComponent.h"
 #include "MediaPlateCustomization.h"
 #include "MediaPlateEditorStyle.h"
@@ -31,6 +33,7 @@ void FMediaPlateEditorModule::StartupModule()
 	FMediaPlateEditorStyle::Get();
 
 	FMediaPlateEditorCommands::Register();
+	RegisterPlacementModeItems();
 
 	// Register customizations.
 	MediaPlateName = UMediaPlateComponent::StaticClass()->GetFName();
@@ -70,6 +73,7 @@ void FMediaPlateEditorModule::ShutdownModule()
 	{
 		SequencerModulePtr->UnRegisterTrackEditor(TrackEditorBindingHandle);
 	}
+	UnregisterPlacementModeItems();
 
 	// Unregister customizations.
 	FPropertyEditorModule& PropertyModule = FModuleManager::LoadModuleChecked<FPropertyEditorModule>("PropertyEditor");
@@ -135,6 +139,99 @@ bool FMediaPlateEditorModule::RemoveMediaSourceFromDragDropCache(UMediaSource* M
 	return bIsInCache;
 }
 
+const FPlacementCategoryInfo* FMediaPlateEditorModule::GetMediaCategoryRegisteredInfo() const
+{
+	const IPlacementModeModule& PlacementModeModule = IPlacementModeModule::Get();
+	const FName PlacementModeCategoryHandle = TEXT("MediaPlate");
+
+	if (const FPlacementCategoryInfo* RegisteredInfo = 
+		PlacementModeModule.GetRegisteredPlacementCategory(PlacementModeCategoryHandle))
+	{
+		return RegisteredInfo;
+	}
+
+	const ISlateStyle* Style = &FMediaPlateEditorStyle::Get().Get();
+	if (Style != nullptr)
+	{
+		FPlacementCategoryInfo Info(
+			LOCTEXT("MediaPlateCategoryName", "Media Plate"),
+			FSlateIcon(Style->GetStyleSetName(), "ClassIcon.MediaPlate"),
+			PlacementModeCategoryHandle,
+			TEXT("MediaPlate"),
+			26 // Determines where the category shows up in the list with respect to the others.
+		);
+
+		IPlacementModeModule::Get().RegisterPlacementCategory(Info);
+
+		return PlacementModeModule.GetRegisteredPlacementCategory(PlacementModeCategoryHandle);
+	}
+	else
+	{
+		return nullptr;
+	}
+}
+
+void FMediaPlateEditorModule::RegisterPlacementModeItems()
+{
+	auto RegisterPlaceActors = [this]() -> void
+	{
+		if (!GEditor)
+		{
+			return;
+		}
+
+		const FPlacementCategoryInfo* Info = GetMediaCategoryRegisteredInfo();
+
+		if (!Info)
+		{
+			UE_LOG(LogMediaPlateEditor, Warning, TEXT("Could not find or create Media Place Actor Category"));
+			return;
+		}
+
+		// Register the Checkerboard
+		PlaceActors.Add(IPlacementModeModule::Get().RegisterPlaceableItem(Info->UniqueHandle, MakeShared<FPlaceableItem>(
+			*AMediaPlate::StaticClass(),
+			FAssetData(AMediaPlate::StaticClass()),
+			NAME_None,
+			NAME_None,
+			TOptional<FLinearColor>(),
+			TOptional<int32>(),
+			LOCTEXT("PlacementMode", "MediaPlate")
+		)));
+	};
+
+	if (FApp::CanEverRender())
+	{
+		if (GEngine && GEngine->IsInitialized())
+		{
+			RegisterPlaceActors();
+		}
+		else
+		{
+			FCoreDelegates::OnPostEngineInit.AddLambda(RegisterPlaceActors);
+		}
+	}
+}
+
+
+void FMediaPlateEditorModule::UnregisterPlacementModeItems()
+{
+	if (IPlacementModeModule::IsAvailable())
+	{
+		IPlacementModeModule& PlacementModeModule = IPlacementModeModule::Get();
+
+		for (TOptional<FPlacementModeID>& PlaceActor : PlaceActors)
+		{
+			if (PlaceActor.IsSet())
+			{
+				PlacementModeModule.UnregisterPlaceableItem(*PlaceActor);
+			}
+		}
+	}
+
+	PlaceActors.Empty();
+}
+
 void FMediaPlateEditorModule::OnPostEngineInit()
 {
 	UEditorAssetSubsystem* EditorAssetSubsystem = GEditor->GetEditorSubsystem<UEditorAssetSubsystem>();
diff --git a/Engine/Plugins/Media/MediaPlate/Source/MediaPlateEditor/Public/MediaPlateEditorModule.h b/Engine/Plugins/Media/MediaPlate/Source/MediaPlateEditor/Public/MediaPlateEditorModule.h
--- a/Engine/Plugins/Media/MediaPlate/Source/MediaPlateEditor/Public/MediaPlateEditorModule.h
+++ b/Engine/Plugins/Media/MediaPlate/Source/MediaPlateEditor/Public/MediaPlateEditorModule.h
@@ -10,9 +10,11 @@
 #include "UObject/ObjectPtr.h"
 
 class FDragDropEvent;
+class FPlacementModeID;
 class UMediaPlateComponent;
 class UMediaSource;
 struct FAssetData;
+struct FPlacementCategoryInfo;
 
 /** Log category for this module. */
 DECLARE_LOG_CATEGORY_EXTERN(LogMediaPlateEditor, Log, All);
@@ -58,6 +60,23 @@ private:
 	TMap<FString, TWeakObjectPtr<UMediaSource>> MapFileToMediaSource;
 	/** Caches if real time viewports are enabled. */
 	bool bIsRealTimeViewportsEnabled = false;
+	/** Stores what we have aded to the placement module. */
+	TArray<TOptional<FPlacementModeID>> PlaceActors;
+
+	/**
+	 * Gathers the Info on the Media Plate Place Actors Category.
+	 */
+	const FPlacementCategoryInfo* GetMediaCategoryRegisteredInfo() const;
+
+	/**
+	 * Register items to show up in the Place Actors panel.
+	 */
+	void RegisterPlacementModeItems();
+
+	/**
+	 * Unregister items in Place Actors panel.
+	 */
+	void UnregisterPlacementModeItems();
 
 	/**
 	 * Called after the engine has initialized.
