MediaPlate: Simplified categories in the detail pane.
#rnx

Merging //UE5/Release-5.2 to Release-Engine-Staging (//UE5/Release-Engine-Staging)

diff --git a/Engine/Plugins/Media/MediaPlate/Source/MediaPlate/Public/MediaPlateComponent.h b/Engine/Plugins/Media/MediaPlate/Source/MediaPlate/Public/MediaPlateComponent.h
--- a/Engine/Plugins/Media/MediaPlate/Source/MediaPlate/Public/MediaPlateComponent.h
+++ b/Engine/Plugins/Media/MediaPlate/Source/MediaPlate/Public/MediaPlateComponent.h
@@ -152,19 +152,19 @@ public:
 	void SetLoop(bool bInLoop);
 
 	/** If set then play when opening the media. */
-	UPROPERTY(EditAnywhere, BlueprintReadWrite, Category = "MediaPlate")
+	UPROPERTY(EditAnywhere, BlueprintReadWrite, Category = "Control")
 	bool bPlayOnOpen = true;
 
 	/** If set then start playing when this object is active. */
-	UPROPERTY(EditAnywhere, Category = "MediaPlate")
+	UPROPERTY(EditAnywhere, Category = "Control")
 	bool bAutoPlay = true;
 
 	/** If set then enable audio. */
-	UPROPERTY(EditAnywhere, Category = "MediaPlate")
+	UPROPERTY(EditAnywhere, Category = "Control")
 	bool bEnableAudio = false;
 
 	/** What time to start playing from (in seconds). */
-	UPROPERTY(EditAnywhere, BlueprintReadWrite, Category = "MediaPlate", meta = (ClampMin = "0.0"))
+	UPROPERTY(EditAnywhere, BlueprintReadWrite, Category = "Control", meta = (ClampMin = "0.0"))
 	float StartTime = 0.0f;
 
 	/** Holds the component to play sound. */
@@ -320,11 +320,11 @@ private:
 	float CurrentRate = 0.0f;
 
 	/** If true then only allow playback when the media plate is visible. */
-	UPROPERTY(EditAnywhere, BlueprintReadOnly, Category = "MediaPlate", meta = (AllowPrivateAccess = true))
+	UPROPERTY(EditAnywhere, BlueprintReadOnly, Category = "Control", meta = (AllowPrivateAccess = true))
 	bool bPlayOnlyWhenVisible = false;
 
 	/** If set then loop when we reach the end. */
-	UPROPERTY(EditAnywhere, Blueprintgetter = GetLoop, BlueprintSetter = SetLoop, Category = "MediaPlate", meta = (AllowPrivateAccess = true))
+	UPROPERTY(EditAnywhere, Blueprintgetter = GetLoop, BlueprintSetter = SetLoop, Category = "Control", meta = (AllowPrivateAccess = true))
 	bool bLoop = true;
 
 	/** Visible mips and tiles calculation mode for the supported mesh types in MediaPlate. (Player restart on change.) */
diff --git a/Engine/Plugins/Media/MediaPlate/Source/MediaPlateEditor/Private/MediaPlateCustomization.cpp b/Engine/Plugins/Media/MediaPlate/Source/MediaPlateEditor/Private/MediaPlateCustomization.cpp
--- a/Engine/Plugins/Media/MediaPlate/Source/MediaPlateEditor/Private/MediaPlateCustomization.cpp
+++ b/Engine/Plugins/Media/MediaPlate/Source/MediaPlateEditor/Private/MediaPlateCustomization.cpp
@@ -65,9 +65,11 @@ void FMediaPlateCustomization::CustomizeDetails(IDetailLayoutBuilder& DetailBuil
 	IDetailCategoryBuilder& StaticMeshCategory = DetailBuilder.EditCategory("StaticMesh");
 	StaticMeshCategory.SetCategoryVisibility(false);
 
+	IDetailCategoryBuilder& ControlCategory = DetailBuilder.EditCategory("Control");
 	IDetailCategoryBuilder& MediaPlateCategory = DetailBuilder.EditCategory("MediaPlate");
 	// Place this at the top.
-	MediaPlateCategory.SetSortOrder(0);
+	ControlCategory.SetSortOrder(0);
+	MediaPlateCategory.SetSortOrder(1);
 
 	// Get objects we are editing.
 	TArray<TWeakObjectPtr<UObject>> Objects;
@@ -188,7 +190,7 @@ void FMediaPlateCustomization::CustomizeDetails(IDetailLayoutBuilder& DetailBuil
 		];
 
 	// Add media control buttons.
-	MediaPlateCategory.AddCustomRow(LOCTEXT("MediaPlateControls", "MediaPlate Controls"))
+	ControlCategory.AddCustomRow(LOCTEXT("MediaPlateControls", "MediaPlate Controls"))
 		[
 			SNew(SHorizontalBox)
 
diff --git a/Engine/Plugins/Media/MediaPlate/Source/MediaPlateEditor/Private/MediaPlateEditorModule.cpp b/Engine/Plugins/Media/MediaPlate/Source/MediaPlateEditor/Private/MediaPlateEditorModule.cpp
--- a/Engine/Plugins/Media/MediaPlate/Source/MediaPlateEditor/Private/MediaPlateEditorModule.cpp
+++ b/Engine/Plugins/Media/MediaPlate/Source/MediaPlateEditor/Private/MediaPlateEditorModule.cpp
@@ -34,6 +34,7 @@ void FMediaPlateEditorModule::StartupModule()
 
 	FMediaPlateEditorCommands::Register();
 	RegisterPlacementModeItems();
+	RegisterSectionMappings();
 
 	// Register customizations.
 	MediaPlateName = UMediaPlateComponent::StaticClass()->GetFName();
@@ -213,6 +214,15 @@ void FMediaPlateEditorModule::RegisterPlacementModeItems()
 	}
 }
 
+void FMediaPlateEditorModule::RegisterSectionMappings()
+{
+	FPropertyEditorModule& PropertyModule = FModuleManager::LoadModuleChecked<FPropertyEditorModule>("PropertyEditor");
+
+	// Media Plate.
+	TSharedRef<FPropertySection> Section = PropertyModule.FindOrCreateSection("Object", "General", LOCTEXT("General", "General"));
+	Section->AddCategory("Control");
+	Section->AddCategory("MediaPlate");
+}
 
 void FMediaPlateEditorModule::UnregisterPlacementModeItems()
 {
diff --git a/Engine/Plugins/Media/MediaPlate/Source/MediaPlateEditor/Public/MediaPlateEditorModule.h b/Engine/Plugins/Media/MediaPlate/Source/MediaPlateEditor/Public/MediaPlateEditorModule.h
--- a/Engine/Plugins/Media/MediaPlate/Source/MediaPlateEditor/Public/MediaPlateEditorModule.h
+++ b/Engine/Plugins/Media/MediaPlate/Source/MediaPlateEditor/Public/MediaPlateEditorModule.h
@@ -73,6 +73,11 @@ private:
 	 */
 	void RegisterPlacementModeItems();
 
+	/**
+	 * Register which categories belong to which section.
+	 */
+	void RegisterSectionMappings();
+
 	/**
 	 * Unregister items in Place Actors panel.
 	 */
