MediaPlate: Transform category is now at the top.
#rnx

diff --git a/Engine/Plugins/Media/MediaPlate/Source/MediaPlateEditor/Private/MediaPlateCustomization.cpp b/Engine/Plugins/Media/MediaPlate/Source/MediaPlateEditor/Private/MediaPlateCustomization.cpp
--- a/Engine/Plugins/Media/MediaPlate/Source/MediaPlateEditor/Private/MediaPlateCustomization.cpp
+++ b/Engine/Plugins/Media/MediaPlate/Source/MediaPlateEditor/Private/MediaPlateCustomization.cpp
@@ -67,9 +67,9 @@ void FMediaPlateCustomization::CustomizeDetails(IDetailLayoutBuilder& DetailBuil
 
 	IDetailCategoryBuilder& ControlCategory = DetailBuilder.EditCategory("Control");
 	IDetailCategoryBuilder& MediaPlateCategory = DetailBuilder.EditCategory("MediaPlate");
-	// Place this at the top.
-	ControlCategory.SetSortOrder(0);
-	MediaPlateCategory.SetSortOrder(1);
+	// The sort order number places these categories after the transform in the detail pane.
+	ControlCategory.SetSortOrder(2010);
+	MediaPlateCategory.SetSortOrder(2020);
 
 	// Get objects we are editing.
 	TArray<TWeakObjectPtr<UObject>> Objects;
