Close MediaPlate editor when owning actor is deleted to fix crash

diff --git a/Engine/Plugins/Media/MediaPlate/Source/MediaPlateEditor/Private/Toolkits/MediaPlateEditorToolkit.cpp b/Engine/Plugins/Media/MediaPlate/Source/MediaPlateEditor/Private/Toolkits/MediaPlateEditorToolkit.cpp
--- a/Engine/Plugins/Media/MediaPlate/Source/MediaPlateEditor/Private/Toolkits/MediaPlateEditorToolkit.cpp
+++ b/Engine/Plugins/Media/MediaPlate/Source/MediaPlateEditor/Private/Toolkits/MediaPlateEditorToolkit.cpp
@@ -3,7 +3,7 @@
 #include "Toolkits/MediaPlateEditorToolkit.h"
 
 #include "Editor.h"
-#include "Styling/AppStyle.h"
+#include "Engine/Engine.h"
 #include "EditorReimportHandler.h"
 #include "Framework/MultiBox/MultiBoxBuilder.h"
 #include "MediaPlate.h"
@@ -13,6 +13,7 @@
 #include "MediaPlaylist.h"
 #include "Models/MediaPlateEditorCommands.h"
 #include "SlateOptMacros.h"
+#include "Styling/AppStyle.h"
 #include "Widgets/Docking/SDockTab.h"
 #include "Widgets/SMediaPlateEditorDetails.h"
 #include "Widgets/SMediaPlateEditorMediaDetails.h"
@@ -43,7 +44,8 @@ FMediaPlateEditorToolkit::~FMediaPlateEditorToolkit()
 {
 	FReimportManager::Instance()->OnPreReimport().RemoveAll(this);
 	FReimportManager::Instance()->OnPostReimport().RemoveAll(this);
-
+	
+	GEngine->OnLevelActorDeleted().RemoveAll(this);
 	GEditor->UnregisterForUndo(this);
 }
 
@@ -133,6 +135,8 @@ void FMediaPlateEditorToolkit::Initialize(UMediaPlateComponent* InMediaPlate, co
 	{
 		EditorModule->MediaPlateStartedPlayback(MediaPlate);
 	}
+
+	GEngine->OnLevelActorDeleted().AddSP(this, &FMediaPlateEditorToolkit::OnActorDeleted);
 }
 
 /* FAssetEditorToolkit interface
@@ -427,4 +431,13 @@ TSharedRef<SDockTab> FMediaPlateEditorToolkit::HandleTabManagerSpawnTab(const FS
 		];
 }
 
+void FMediaPlateEditorToolkit::OnActorDeleted(AActor* Actor)
+{
+	// If our actor owning actor is removed, we need to close. Our MediaPlate referenced will be cleared by GC and cause a crash later.
+	if (MediaPlate && MediaPlate->GetOwner() == Actor)
+	{
+		CloseWindow(EAssetEditorCloseReason::AssetUnloadingOrInvalid);
+	}
+}
+
 #undef LOCTEXT_NAMESPACE
diff --git a/Engine/Plugins/Media/MediaPlate/Source/MediaPlateEditor/Private/Toolkits/MediaPlateEditorToolkit.h b/Engine/Plugins/Media/MediaPlate/Source/MediaPlateEditor/Private/Toolkits/MediaPlateEditorToolkit.h
--- a/Engine/Plugins/Media/MediaPlate/Source/MediaPlateEditor/Private/Toolkits/MediaPlateEditorToolkit.h
+++ b/Engine/Plugins/Media/MediaPlate/Source/MediaPlateEditor/Private/Toolkits/MediaPlateEditorToolkit.h
@@ -29,7 +29,6 @@ public:
 
 	/** Virtual destructor. */
 	virtual ~FMediaPlateEditorToolkit();
-
 	/**
 	 * Initializes the editor tool kit.
 	 *
@@ -94,4 +93,6 @@ private:
 
 	/** Pointer to the style set to use for toolkits. */
 	TSharedRef<ISlateStyle> Style;
+	
+	void OnActorDeleted(AActor* Actor);
 };
