MediaPlate: Fix mesh range issues.

diff --git a/Engine/Plugins/Media/MediaPlate/Source/MediaPlateEditor/Private/MediaPlateCustomization.cpp b/Engine/Plugins/Media/MediaPlate/Source/MediaPlateEditor/Private/MediaPlateCustomization.cpp
--- a/Engine/Plugins/Media/MediaPlate/Source/MediaPlateEditor/Private/MediaPlateCustomization.cpp
+++ b/Engine/Plugins/Media/MediaPlate/Source/MediaPlateEditor/Private/MediaPlateCustomization.cpp
@@ -872,7 +872,7 @@ TOptional<float> FMediaPlateCustomization::GetLetterboxAspectRatio() const
 
 void FMediaPlateCustomization::SetMeshHorizontalRange(float HorizontalRange)
 {
-	HorizontalRange = FMath::Clamp(HorizontalRange, 0.0f, 360.0f);
+	HorizontalRange = FMath::Clamp(HorizontalRange, 1.0f, 360.0f);
 	TOptional VerticalRange = GetMeshVerticalRange();
 	if (VerticalRange.IsSet())
 	{
@@ -898,7 +898,7 @@ TOptional<float> FMediaPlateCustomization::GetMeshHorizontalRange() const
 
 void FMediaPlateCustomization::SetMeshVerticalRange(float VerticalRange)
 {
-	VerticalRange = FMath::Clamp(VerticalRange, 0.0f, 180.0f);
+	VerticalRange = FMath::Clamp(VerticalRange, 1.0f, 180.0f);
 	TOptional HorizontalRange = GetMeshHorizontalRange();
 	if (HorizontalRange.IsSet())
 	{
@@ -924,6 +924,8 @@ TOptional<float> FMediaPlateCustomization::GetMeshVerticalRange() const
 
 void FMediaPlateCustomization::SetMeshRange(FVector2D Range)
 {
+	const FScopedTransaction Transaction(LOCTEXT("SetMeshRange", "Media Plate Set Mesh Range"));
+
 	// Loop through all our objects.
 	for (const TWeakObjectPtr<UMediaPlateComponent>& MediaPlatePtr : MediaPlatesList)
 	{
@@ -932,6 +934,7 @@ void FMediaPlateCustomization::SetMeshRange(FVector2D Range)
 		{
 			if (MediaPlate->GetMeshRange() != Range)
 			{
+				MediaPlate->Modify();
 				MediaPlate->SetMeshRange(Range);
 				SetSphereMesh(MediaPlate);
 			}
