MediaPlate: Fixed player not closing when deleting the media plate.

diff --git a/Engine/Plugins/Media/MediaPlate/Source/MediaPlate/Private/MediaPlateComponent.cpp b/Engine/Plugins/Media/MediaPlate/Source/MediaPlate/Private/MediaPlateComponent.cpp
--- a/Engine/Plugins/Media/MediaPlate/Source/MediaPlate/Private/MediaPlateComponent.cpp
+++ b/Engine/Plugins/Media/MediaPlate/Source/MediaPlate/Private/MediaPlateComponent.cpp
@@ -203,6 +203,8 @@ void UMediaPlateComponent::OnUnregister()
 	Super::OnUnregister();
 
 	UnregisterWithMediaTextureTracker();
+
+	Close();
 }
 
 UMediaPlayer* UMediaPlateComponent::GetMediaPlayer()
