Media Plate: Exposed Media Texture's properties via Media Plate interface.

diff --git a/Engine/Plugins/Media/MediaPlate/Source/MediaPlate/Private/MediaPlateComponent.cpp b/Engine/Plugins/Media/MediaPlate/Source/MediaPlate/Private/MediaPlateComponent.cpp
--- a/Engine/Plugins/Media/MediaPlate/Source/MediaPlate/Private/MediaPlateComponent.cpp
+++ b/Engine/Plugins/Media/MediaPlate/Source/MediaPlate/Private/MediaPlateComponent.cpp
@@ -33,6 +33,36 @@ namespace UE::MediaPlateComponent
 		ForceUpdateResource = 0x2,
 	};
 	ENUM_CLASS_FLAGS(ESetUpTexturesFlags);
+
+	// Runs through media textures and sets Media Plate settings corresponding to Media Texture. 
+	void ApplyMediaTextureMipGenProperties(const FMediaTextureResourceSettings MediaTextureSettings, const TArray<TObjectPtr<UMediaTexture>>& MediaTextures)
+	{
+		for (TObjectPtr<UMediaTexture> MediaTexture : MediaTextures)
+		{
+			if (MediaTexture != nullptr
+				&& (MediaTexture->EnableGenMips != MediaTextureSettings.bEnableGenMips || MediaTexture->NumMips != MediaTextureSettings.CurrentNumMips))
+			{
+				MediaTexture->EnableGenMips = MediaTextureSettings.bEnableGenMips;
+				MediaTexture->NumMips = MediaTextureSettings.CurrentNumMips;
+				MediaTexture->UpdateResource();
+			}
+		}
+	}
+
+	void EnsureMediaTexturePropertiesInSync(const FMediaTextureResourceSettings MediaTextureSettings, const TArray<TObjectPtr<UMediaTexture>>& MediaTextures)
+	{
+#if !UE_BUILD_SHIPPING
+		for (TObjectPtr<UMediaTexture> MediaTexture : MediaTextures)
+		{
+			if (MediaTexture != nullptr)
+			{
+				bool bMediaTextureMipGenPropertiesInSync = MediaTexture->EnableGenMips == MediaTextureSettings.bEnableGenMips && MediaTexture->NumMips == MediaTextureSettings.CurrentNumMips;
+				ensureMsgf(bMediaTextureMipGenPropertiesInSync, TEXT("Mip Generation properties set on Media Plate are different from the properties set on Media Texture. \n\
+					Media Texture mip generation properites are not meant to be modified directly."));
+			}
+		}
+#endif
+	}
 };
 
 /**
@@ -88,6 +118,9 @@ UMediaPlateComponent::UMediaPlateComponent(const FObjectInitializer& ObjectIniti
 
 	// Default to plane since AMediaPlate defaults to SM_MediaPlateScreen
 	VisibleMipsTilesCalculations = EMediaTextureVisibleMipsTiles::Plane;
+
+	MediaTextureSettings.bEnableGenMips = false;
+	MediaTextureSettings.CurrentNumMips = 1;
 }
 
 #if WITH_EDITOR
@@ -104,6 +137,9 @@ void UMediaPlateComponent::PostLoad()
 		}
 		MediaTexture_DEPRECATED = nullptr;
 	}
+
+	UE::MediaPlateComponent::ApplyMediaTextureMipGenProperties(MediaTextureSettings, MediaTextures);
+
 }
 #endif // WITH_EDITOR
 
@@ -177,6 +213,9 @@ void UMediaPlateComponent::TickComponent(float DeltaTime, enum ELevelTick TickTy
 {
 	Super::TickComponent(DeltaTime, TickType, ThisTickFunction);
 
+	// Making sure that Media Plate enforces the same settings to Media textures in case these settings were modified externally.
+	UE::MediaPlateComponent::EnsureMediaTexturePropertiesInSync(MediaTextureSettings, MediaTextures);
+
 	if (MediaPlayer != nullptr)
 	{
 		if ((CurrentRate != 0.0f) || (bWantsToPlayWhenVisible))
@@ -1002,6 +1041,14 @@ void UMediaPlateComponent::SetUpTextures(UE::MediaPlateComponent::ESetUpTextures
 		{
 			bool bApplyTextureUpdate = false;
 
+			if (MediaTexture->EnableGenMips != MediaTextureSettings.bEnableGenMips
+				|| MediaTexture->NumMips != MediaTextureSettings.CurrentNumMips)
+			{
+				MediaTexture->EnableGenMips = MediaTextureSettings.bEnableGenMips;
+				MediaTexture->NumMips = MediaTextureSettings.CurrentNumMips;
+				bApplyTextureUpdate = true;
+			}
+
 			if (FMath::IsNearlyEqual(MediaTexture->GetMipMapBias(), MipMapBias) == false)
 			{
 				MediaTexture->SetMipMapBias(MipMapBias);
@@ -1212,6 +1259,12 @@ void UMediaPlateComponent::PostEditChangeProperty(FPropertyChangedEvent& Propert
 			MediaTextureTrackerObject->MipLevelToUpscale = bEnableMipMapUpscaling ? MipLevelToUpscale : -1;
 		}
 	}
+	else if (PropertyChangedEvent.GetPropertyName() == GET_MEMBER_NAME_CHECKED(FMediaTextureResourceSettings, bEnableGenMips)
+		|| PropertyChangedEvent.GetPropertyName() == GET_MEMBER_NAME_CHECKED(FMediaTextureResourceSettings, CurrentNumMips))
+	{
+		UE::MediaPlateComponent::ApplyMediaTextureMipGenProperties(MediaTextureSettings, MediaTextures);
+		RestartPlayer();
+	}
 }
 
 void UMediaPlateComponent::SwitchStates(EMediaPlateEventState State)
diff --git a/Engine/Plugins/Media/MediaPlate/Source/MediaPlate/Public/MediaPlateComponent.h b/Engine/Plugins/Media/MediaPlate/Source/MediaPlate/Public/MediaPlateComponent.h
--- a/Engine/Plugins/Media/MediaPlate/Source/MediaPlate/Public/MediaPlateComponent.h
+++ b/Engine/Plugins/Media/MediaPlate/Source/MediaPlate/Public/MediaPlateComponent.h
@@ -38,6 +38,25 @@ enum class EMediaPlateEventState : uint8
 	MAX
 };
 
+/**
+ * This struct is used to expose Media Texture settings via Media Plate Component and is a mirror of some
+ * of the settings.
+ */
+USTRUCT()
+struct FMediaTextureResourceSettings
+{
+	GENERATED_USTRUCT_BODY()
+
+	/** Enable mips generation */
+	UPROPERTY(EditAnywhere, Category = "Media Texture Resource Settings")
+	bool bEnableGenMips;
+
+	/** Current number of mips to be generated as output */
+	UPROPERTY(EditAnywhere, Category = "Media Texture Resource Settings")
+	uint8 CurrentNumMips;
+};
+
+
 /**
  * This is a component for AMediaPlate that can play and show media in the world.
  */
@@ -373,6 +392,10 @@ private:
 	UPROPERTY(Instanced)
 	TArray<TObjectPtr<UMediaTexture>> MediaTextures;
 
+	/** Exposes Media Texture settings via Media Plate component. */
+	UPROPERTY(EditAnywhere, AdvancedDisplay, Category = "MediaPlate")
+	FMediaTextureResourceSettings MediaTextureSettings;
+
 	/** This component's media player */
 	UPROPERTY(Instanced)
 	TObjectPtr<UMediaPlayer> MediaPlayer;
