Fix uninitialized structs
#rnx

diff --git a/Engine/Plugins/Media/MediaPlate/Source/MediaPlate/Public/MediaPlateComponent.h b/Engine/Plugins/Media/MediaPlate/Source/MediaPlate/Public/MediaPlateComponent.h
--- a/Engine/Plugins/Media/MediaPlate/Source/MediaPlate/Public/MediaPlateComponent.h
+++ b/Engine/Plugins/Media/MediaPlate/Source/MediaPlate/Public/MediaPlateComponent.h
@@ -49,11 +49,11 @@ struct FMediaTextureResourceSettings
 
 	/** Enable mips generation */
 	UPROPERTY(EditAnywhere, Category = "Media Texture Resource Settings")
-	bool bEnableGenMips;
+	bool bEnableGenMips = false;
 
 	/** Current number of mips to be generated as output */
 	UPROPERTY(EditAnywhere, Category = "Media Texture Resource Settings")
-	uint8 CurrentNumMips;
+	uint8 CurrentNumMips = 0;
 };
 
 
