ImgMedia: Added default values to the Media texture properties struct.

diff --git a/Engine/Plugins/Media/MediaPlate/Source/MediaPlate/Public/MediaPlateComponent.h b/Engine/Plugins/Media/MediaPlate/Source/MediaPlate/Public/MediaPlateComponent.h
--- a/Engine/Plugins/Media/MediaPlate/Source/MediaPlate/Public/MediaPlateComponent.h
+++ b/Engine/Plugins/Media/MediaPlate/Source/MediaPlate/Public/MediaPlateComponent.h
@@ -53,7 +53,7 @@ struct FMediaTextureResourceSettings
 
 	/** Current number of mips to be generated as output */
 	UPROPERTY(EditAnywhere, Category = "Media Texture Resource Settings")
-	uint8 CurrentNumMips = 0;
+	uint8 CurrentNumMips = 1;
 };
 
 
