MediaPlate: Enable adaptive pole mip upscaling by default.

#rnx

diff --git a/Engine/Plugins/Media/MediaPlate/Source/MediaPlate/Public/MediaPlateComponent.h b/Engine/Plugins/Media/MediaPlate/Source/MediaPlate/Public/MediaPlateComponent.h
--- a/Engine/Plugins/Media/MediaPlate/Source/MediaPlate/Public/MediaPlateComponent.h
+++ b/Engine/Plugins/Media/MediaPlate/Source/MediaPlate/Public/MediaPlateComponent.h
@@ -368,7 +368,7 @@ private:
 
 	/** If true then Media Plate will attempt to load and upscale lower quality mips and display those at the poles (Sphere object only). */
 	UPROPERTY(EditAnywhere, AdvancedDisplay, Category = "MediaPlate", meta = (EditCondition = "VisibleMipsTilesCalculations == EMediaTextureVisibleMipsTiles::Sphere", AllowPrivateAccess = true))
-	bool bAdaptivePoleMipUpscaling = false;
+	bool bAdaptivePoleMipUpscaling = true;
 
 	/** If > 0, then this is the aspect ratio of our screen and 
 	 * letterboxes will be added if the media is smaller than the screen. */
