MediaPlate: Add media texture support for overlay materials.

diff --git a/Engine/Plugins/Media/MediaPlate/Source/MediaPlate/Private/MediaPlate.cpp b/Engine/Plugins/Media/MediaPlate/Source/MediaPlate/Private/MediaPlate.cpp
--- a/Engine/Plugins/Media/MediaPlate/Source/MediaPlate/Private/MediaPlate.cpp
+++ b/Engine/Plugins/Media/MediaPlate/Source/MediaPlate/Private/MediaPlate.cpp
@@ -107,7 +107,7 @@ void AMediaPlate::BeginDestroy()
 
 UMaterialInterface* AMediaPlate::GetCurrentMaterial() const
 {
-	if (StaticMeshComponent != nullptr)
+	if (IsValid(StaticMeshComponent))
 	{
 		return StaticMeshComponent->GetMaterial(0);
 	}
@@ -115,6 +115,16 @@ UMaterialInterface* AMediaPlate::GetCurrentMaterial() const
 	return nullptr;
 }
 
+UMaterialInterface* AMediaPlate::GetCurrentOverlayMaterial() const
+{
+	if (IsValid(StaticMeshComponent))
+	{
+		return StaticMeshComponent->GetOverlayMaterial();
+	}
+
+	return nullptr;
+}
+
 #if WITH_EDITOR
 
 void AMediaPlate::UseDefaultMaterial()
@@ -132,6 +142,40 @@ void AMediaPlate::ApplyCurrentMaterial()
 	{
 		ApplyMaterial(MaterialInterface);
 	}
+
+	UMaterialInterface* OverlayMaterialInterface = GetCurrentOverlayMaterial();
+
+	if ((OverlayMaterialInterface != nullptr) && (LastOverlayMaterial != OverlayMaterialInterface))
+	{
+		ApplyOverlayMaterial(OverlayMaterialInterface);
+	}
+}
+
+UMaterialInterface* AMediaPlate::CreateMaterialInstanceConstant(UMaterialInterface* InMaterial)
+{
+	// Change M_ to MI_ in material name and then generate a unique one.
+	FString MaterialName = InMaterial->GetName();
+	if (MaterialName.StartsWith(TEXT("M_")))
+	{
+		MaterialName.InsertAt(1, TEXT("I"));
+	}
+	FName MaterialUniqueName = MakeUniqueObjectName(StaticMeshComponent, UMaterialInstanceConstant::StaticClass(),
+		FName(*MaterialName));
+
+	// Create instance.
+	UMaterialInstanceConstant* MaterialInstance =
+		NewObject<UMaterialInstanceConstant>(StaticMeshComponent, MaterialUniqueName, RF_Transactional);
+	MaterialInstance->SetParentEditorOnly(InMaterial);
+	MaterialInstance->CopyMaterialUniformParametersEditorOnly(InMaterial);
+	MaterialInstance->SetTextureParameterValueEditorOnly(
+		FMaterialParameterInfo(MediaTextureName),
+		MediaPlateComponent->GetMediaTexture());
+	MaterialInstance->PostEditChange();
+
+	// We force call post-load to indirectly call UpdateParameters() (for integration with VPUtilities plugin).
+	MaterialInstance->PostLoad();
+
+	return MaterialInstance;
 }
 
 void AMediaPlate::ApplyMaterial(UMaterialInterface* Material)
@@ -170,30 +214,7 @@ void AMediaPlate::ApplyMaterial(UMaterialInterface* Material)
 			else
 			{
 				MediaPlateComponent->SetNumberOfTextures(1);
-
-				// Change M_ to MI_ in material name and then generate a unique one.
-				FString MaterialName = Material->GetName();
-				if (MaterialName.StartsWith(TEXT("M_")))
-				{
-					MaterialName.InsertAt(1, TEXT("I"));
-				}
-				FName MaterialUniqueName = MakeUniqueObjectName(StaticMeshComponent, UMaterialInstanceConstant::StaticClass(),
-					FName(*MaterialName));
-
-				// Create instance.
-				UMaterialInstanceConstant* MaterialInstance =
-					NewObject<UMaterialInstanceConstant>(StaticMeshComponent, MaterialUniqueName, RF_Transactional);
-				MaterialInstance->SetParentEditorOnly(Material);
-				MaterialInstance->CopyMaterialUniformParametersEditorOnly(Material);
-				MaterialInstance->SetTextureParameterValueEditorOnly(
-					FMaterialParameterInfo(MediaTextureName),
-					MediaPlateComponent->GetMediaTexture());
-				MaterialInstance->PostEditChange();
-
-				// We force call post-load to indirectly call UpdateParameters() (for integration with VPUtilities plugin).
-				MaterialInstance->PostLoad();
-
-				Result = MaterialInstance;
+				Result = CreateMaterialInstanceConstant(Material);
 			}
 
 			// Update static mesh.
@@ -208,6 +229,50 @@ void AMediaPlate::ApplyMaterial(UMaterialInterface* Material)
 	}
 }
 
+void AMediaPlate::ApplyOverlayMaterial(UMaterialInterface* InOverlayMaterial)
+{
+	if (InOverlayMaterial != nullptr && StaticMeshComponent != nullptr)
+	{
+		UMaterialInstanceDynamic* MID = Cast<UMaterialInstanceDynamic>(InOverlayMaterial);
+
+		if (GEditor == nullptr)
+		{
+			if (MID == nullptr)
+			{
+				// Create and set the dynamic material instance.
+				MID = UMaterialInstanceDynamic::Create(InOverlayMaterial, StaticMeshComponent);
+			}
+			StaticMeshComponent->SetOverlayMaterial(MID);
+			SetMIDParameters(MID);
+			LastOverlayMaterial = MID;
+		}
+		else
+		{
+			UMaterialInterface* Result = nullptr;
+
+			if (MID != nullptr)
+			{
+				SetMIDParameters(MID);
+				Result = MID;
+			}
+			else
+			{
+				MediaPlateComponent->SetNumberOfTextures(1);
+				Result = CreateMaterialInstanceConstant(InOverlayMaterial);
+			}
+
+			// Update static mesh.
+			if (Result != nullptr)
+			{
+				StaticMeshComponent->Modify();
+				StaticMeshComponent->SetOverlayMaterial(Result);
+
+				LastOverlayMaterial = Result;
+			}
+		}
+	}
+}
+
 void AMediaPlate::SetMIDParameters(UMaterialInstanceDynamic* InMaterial)
 {
 	InMaterial->SetTextureParameterValue(MediaTextureName, MediaPlateComponent->GetMediaTexture());
diff --git a/Engine/Plugins/Media/MediaPlate/Source/MediaPlate/Private/MediaPlateComponent.cpp b/Engine/Plugins/Media/MediaPlate/Source/MediaPlate/Private/MediaPlateComponent.cpp
--- a/Engine/Plugins/Media/MediaPlate/Source/MediaPlate/Private/MediaPlateComponent.cpp
+++ b/Engine/Plugins/Media/MediaPlate/Source/MediaPlate/Private/MediaPlateComponent.cpp
@@ -1122,6 +1122,8 @@ void UMediaPlateComponent::UpdateTextureLayers()
 {
 	if (AMediaPlate* MediaPlate = GetOwner<AMediaPlate>())
 	{
+		static const FString BaseTextureName = TEXT("MediaTexture");
+
 		if (UMaterialInterface* Material = MediaPlate->GetCurrentMaterial())
 		{
 			UMaterialInstanceDynamic* MID = Cast<UMaterialInstanceDynamic>(Material);
@@ -1129,7 +1131,6 @@ void UMediaPlateComponent::UpdateTextureLayers()
 			{
 				// Go through each layer.
 				int32 MatNumLayers = MediaTextures.Num() / MatNumTexPerLayer;
-				static const FString BaseTextureName = TEXT("MediaTexture");
 				int32 NumLayers = TextureLayers.Num();
 				int32 MaterialLayerIndex = 0;
 				for (int32 LayerIndex = 0; LayerIndex < NumLayers; ++LayerIndex)
@@ -1171,6 +1172,17 @@ void UMediaPlateComponent::UpdateTextureLayers()
 				}
 			}
 		}
+
+		UMaterialInterface* OverlayMaterial = MediaPlate->GetCurrentOverlayMaterial();
+		if (OverlayMaterial != nullptr)
+		{
+			UMaterialInstanceDynamic* MID = Cast<UMaterialInstanceDynamic>(OverlayMaterial);
+			if (MID != nullptr && !MediaTextures.IsEmpty())
+			{
+
+				MID->SetTextureParameterValue(FName(*BaseTextureName), MediaTextures[0]);
+			}
+		}
 	}
 }
 
diff --git a/Engine/Plugins/Media/MediaPlate/Source/MediaPlate/Public/MediaPlate.h b/Engine/Plugins/Media/MediaPlate/Source/MediaPlate/Public/MediaPlate.h
--- a/Engine/Plugins/Media/MediaPlate/Source/MediaPlate/Public/MediaPlate.h
+++ b/Engine/Plugins/Media/MediaPlate/Source/MediaPlate/Public/MediaPlate.h
@@ -44,8 +44,17 @@ public:
 	 * Call this after changing the current material to set it up for media plate.
 	 */
 	void ApplyCurrentMaterial();
+	
+	/**
+	 * Setup the material for media plate use. Automatically called by ApplyCurrentMaterial.
+	 */
 	void ApplyMaterial(UMaterialInterface* InMaterial);
 
+	/**
+	 * Setup the overlay material for media plate use. Automatically called by ApplyCurrentMaterial.
+	 */
+	void ApplyOverlayMaterial(UMaterialInterface* InOverlayMaterial);
+
 	/**
 	 * Sets up parameters (like the texture) that we use in the material.
 	 */
@@ -58,6 +67,9 @@ public:
 	/** Get the current static mesh material, at index 0. */
 	UMaterialInterface* GetCurrentMaterial() const;
 
+	/** Get the current static mesh overlay material, nullptr otherwise. */
+	UMaterialInterface* GetCurrentOverlayMaterial() const;
+
 private:
 	/** Name for our media plate component. */
 	static FLazyName MediaPlateComponentName;
@@ -66,6 +78,7 @@ private:
 
 #if WITH_EDITOR
 	UMaterialInterface* LastMaterial = nullptr;
+	UMaterialInterface* LastOverlayMaterial = nullptr;
 
 	/**
 	 * Called before a level saves
@@ -87,5 +100,10 @@ private:
 	 */
 	void RemoveAssetUserData();
 
+	/**
+	 * Convenience function to apply create a material instance constant for media plate use.
+	 */
+	UMaterialInterface* CreateMaterialInstanceConstant(UMaterialInterface* InMaterial);
+
 #endif // WITH_EDITOR
 };
