MediaPlate: Add player name to media details.

#rnx

diff --git a/Engine/Plugins/Media/MediaPlate/Source/MediaPlateEditor/Private/Widgets/SMediaPlateEditorMediaDetails.cpp b/Engine/Plugins/Media/MediaPlate/Source/MediaPlateEditor/Private/Widgets/SMediaPlateEditorMediaDetails.cpp
--- a/Engine/Plugins/Media/MediaPlate/Source/MediaPlateEditor/Private/Widgets/SMediaPlateEditorMediaDetails.cpp
+++ b/Engine/Plugins/Media/MediaPlate/Source/MediaPlateEditor/Private/Widgets/SMediaPlateEditorMediaDetails.cpp
@@ -34,6 +34,15 @@ void SMediaPlateEditorMediaDetails::Construct(const FArguments& InArgs,
 						[
 							SNew(SVerticalBox)
 
+							// Player name.
+							+ SVerticalBox::Slot()
+								.AutoHeight()
+								.VAlign(VAlign_Center)
+								.Padding(4.0f)
+								[
+									SAssignNew(MediaPlayerName, STextBlock)
+								]
+
 							// Resolution.
 							+ SVerticalBox::Slot()
 								.AutoHeight()
@@ -130,6 +139,7 @@ void SMediaPlateEditorMediaDetails::Tick(const FGeometry& AllottedGeometry, cons
 
 void SMediaPlateEditorMediaDetails::UpdateDetails()
 {
+	FName PlayerName;
 	FString Format;
 	float FrameRate = 0.0f;
 	int32 LODBias = 0;
@@ -146,6 +156,7 @@ void SMediaPlateEditorMediaDetails::UpdateDetails()
 		UMediaPlayer* MediaPlayer = MediaPlate->GetMediaPlayer();
 		if (MediaPlayer != nullptr)
 		{
+			PlayerName = MediaPlayer->GetPlayerName();
 			FrameRate = MediaPlayer->GetVideoTrackFrameRate(INDEX_NONE, INDEX_NONE);
 			Format = MediaPlayer->GetVideoTrackType(INDEX_NONE, INDEX_NONE);
 			FIntPoint NumTiles(EForceInit::ForceInitToZero);
@@ -170,6 +181,8 @@ void SMediaPlateEditorMediaDetails::UpdateDetails()
 	}
 
 	// Update text.
+	MediaPlayerName->SetText(FText::Format(LOCTEXT("Player", "Player: {0}"),
+		FText::FromName(PlayerName)));
 	FormatText->SetText(FText::Format(LOCTEXT("Format", "Format: {0}"),
 		FText::FromString(Format)));
 	FrameRateText->SetText(FText::Format(LOCTEXT("FrameRate", "Frame Rate: {0}"), 
diff --git a/Engine/Plugins/Media/MediaPlate/Source/MediaPlateEditor/Private/Widgets/SMediaPlateEditorMediaDetails.h b/Engine/Plugins/Media/MediaPlate/Source/MediaPlateEditor/Private/Widgets/SMediaPlateEditorMediaDetails.h
--- a/Engine/Plugins/Media/MediaPlate/Source/MediaPlateEditor/Private/Widgets/SMediaPlateEditorMediaDetails.h
+++ b/Engine/Plugins/Media/MediaPlate/Source/MediaPlateEditor/Private/Widgets/SMediaPlateEditorMediaDetails.h
@@ -45,6 +45,7 @@ private:
 	TWeakObjectPtr<UMediaPlateComponent> MediaPlate;
 
 	/** Our widgets. */
+	TSharedPtr<STextBlock> MediaPlayerName;
 	TSharedPtr<STextBlock> FormatText;
 	TSharedPtr<STextBlock> FrameRateText;
 	TSharedPtr<STextBlock> LODBiasText;
