MediaPlate: Depth-aware overlay composite to minimize TSR artifacts.

A contextual menu is now present on media plate objects, and lets users enable required state to bypass common TSR issues with in-scene video playback. While sufficiently useful, note that it is not a panacea: at low screen percentages, edge artifacts can still appear due to jitter & low resolution of scene depth. Also note that as the last translucency pass, it can render over existing translucency.

#virtualized

diff --git a/Engine/Plugins/Media/MediaPlate/Source/MediaPlate/Private/MediaPlate.cpp b/Engine/Plugins/Media/MediaPlate/Source/MediaPlate/Private/MediaPlate.cpp
--- a/Engine/Plugins/Media/MediaPlate/Source/MediaPlate/Private/MediaPlate.cpp
+++ b/Engine/Plugins/Media/MediaPlate/Source/MediaPlate/Private/MediaPlate.cpp
@@ -25,6 +25,27 @@
 FLazyName AMediaPlate::MediaPlateComponentName(TEXT("MediaPlateComponent0"));
 FLazyName AMediaPlate::MediaTextureName("MediaTexture");
 
+namespace UE::MediaPlate::Private
+{
+	void ApplyTranslucencyScreenPercentageCVar(int32 InBasis)
+	{
+		static IConsoleVariable* TranslucencySPBasisCVar = IConsoleManager::Get().FindConsoleVariable(TEXT("r.Translucency.ScreenPercentage.Basis"));
+		if (TranslucencySPBasisCVar->GetInt() != InBasis)
+		{
+			if (InBasis)
+			{
+				UE_LOG(LogMediaPlate, Warning, TEXT("Setting 'r.Translucency.ScreenPercentage.Basis' to 1. For media plates with overlay materials, please apply this console variable permanently to your project."));
+			}
+			else
+			{
+				UE_LOG(LogMediaPlate, Warning, TEXT("Setting 'r.Translucency.ScreenPercentage.Basis' to 0."));
+			}
+
+			TranslucencySPBasisCVar->Set(InBasis);
+		}
+	}
+}
+
 AMediaPlate::AMediaPlate(const FObjectInitializer& ObjectInitializer)
 	: Super(ObjectInitializer)
 {
@@ -95,6 +116,12 @@ void AMediaPlate::PostRegisterAllComponents()
 	{
 		UseDefaultMaterial();		
 	}
+
+	if (IsValid(StaticMeshComponent) &&
+		IsValid(StaticMeshComponent->OverlayMaterial))
+	{
+		UE::MediaPlate::Private::ApplyTranslucencyScreenPercentageCVar(1);
+	}
 	
 	AddAssetUserData();
 #endif // WITH_EDITOR
@@ -132,6 +159,13 @@ void AMediaPlate::UseDefaultMaterial()
 	UMaterial* DefaultMaterial = LoadObject<UMaterial>(NULL, TEXT("/MediaPlate/M_MediaPlate"), NULL, LOAD_None, NULL);
 	
 	ApplyMaterial(DefaultMaterial);
+
+	if (IsValid(StaticMeshComponent))
+	{
+		StaticMeshComponent->SetOverlayMaterial(nullptr);
+		
+		LastOverlayMaterial = nullptr;
+	}
 }
 
 void AMediaPlate::ApplyCurrentMaterial()
@@ -270,6 +304,8 @@ void AMediaPlate::ApplyOverlayMaterial(UMaterialInterface* InOverlayMaterial)
 				LastOverlayMaterial = Result;
 			}
 		}
+
+		UE::MediaPlate::Private::ApplyTranslucencyScreenPercentageCVar(1);
 	}
 }
 
diff --git a/Engine/Plugins/Media/MediaPlate/Source/MediaPlate/Public/MediaPlate.h b/Engine/Plugins/Media/MediaPlate/Source/MediaPlate/Public/MediaPlate.h
--- a/Engine/Plugins/Media/MediaPlate/Source/MediaPlate/Public/MediaPlate.h
+++ b/Engine/Plugins/Media/MediaPlate/Source/MediaPlate/Public/MediaPlate.h
@@ -11,6 +11,11 @@ class UMaterialInterface;
 class UMaterialInstanceDynamic;
 class UMediaPlateComponent;
 
+namespace UE::MediaPlate::Private
+{
+	MEDIAPLATE_API void ApplyTranslucencyScreenPercentageCVar(int32 InBasis);
+}
+
 /**
  * MediaPlate is an actor that can play and show media in the world.
  */
@@ -36,7 +41,7 @@ public:
 #if WITH_EDITOR
 
 	/*
-	 * Call this to change the static mesh to use the default media plate material.
+	 * Call this to change the static mesh to use the default media plate material and reset the overlay material.
 	 */
 	void UseDefaultMaterial();
 
diff --git a/Engine/Plugins/Media/MediaPlate/Source/MediaPlateEditor/Private/MediaPlateEditorModule.cpp b/Engine/Plugins/Media/MediaPlate/Source/MediaPlateEditor/Private/MediaPlateEditorModule.cpp
--- a/Engine/Plugins/Media/MediaPlate/Source/MediaPlateEditor/Private/MediaPlateEditorModule.cpp
+++ b/Engine/Plugins/Media/MediaPlate/Source/MediaPlateEditor/Private/MediaPlateEditorModule.cpp
@@ -7,6 +7,7 @@
 #include "ISequencerModule.h"
 #include "LevelEditor.h"
 #include "MaterialList.h"
+#include "Materials/Material.h"
 #include "MediaPlate.h"
 #include "MediaPlateComponent.h"
 #include "MediaPlateCustomization.h"
@@ -46,6 +47,8 @@ void FMediaPlateEditorModule::StartupModule()
 	ISequencerModule& SequencerModule = FModuleManager::Get().LoadModuleChecked<ISequencerModule>("Sequencer");
 	TrackEditorBindingHandle = SequencerModule.RegisterPropertyTrackEditor<FMediaPlateTrackEditor>();
 
+	RegisterContextMenuExtender();
+
 	// Add bottom extender for material item
 	FMaterialList::OnAddMaterialItemViewExtraBottomWidget.AddLambda([](const TSharedRef<FMaterialItemView>& InMaterialItemView, UActorComponent* InCurrentComponent, IDetailLayoutBuilder& InDetailBuilder, TArray<TSharedPtr<SWidget>>& OutExtensions)
 	{
@@ -60,6 +63,8 @@ void FMediaPlateEditorModule::ShutdownModule()
 {
 	FMaterialList::OnAddMaterialItemViewExtraBottomWidget.RemoveAll(this);
 
+	UnregisterContextMenuExtender();
+
 	if (GEditor != nullptr)
 	{
 		UEditorAssetSubsystem* EditorAssetSubsystem = GEditor->GetEditorSubsystem<UEditorAssetSubsystem>();
@@ -315,6 +320,111 @@ void FMediaPlateEditorModule::ForceRealTimeViewports(const bool bEnable)
 	}
 }
 
+TSharedRef<FExtender> FMediaPlateEditorModule::ExtendLevelViewportContextMenuForMediaPlate(const TSharedRef<FUICommandList> CommandList, TArray<AActor*> SelectedActors)
+{
+	TSharedPtr<FExtender> Extender = MakeShared<FExtender>();
+
+	if (SelectedActors.Num() == 1)
+	{
+		AMediaPlate* MediaPlateActor = Cast<AMediaPlate>(SelectedActors[0]);
+		if (IsValid(MediaPlateActor))
+		{
+			Extender->AddMenuExtension("ActorTypeTools", EExtensionHook::After, CommandList,
+				FMenuExtensionDelegate::CreateLambda([MediaPlateActor](FMenuBuilder& MenuBuilder, AActor* SelectedActor)
+					{
+						const ISlateStyle* Style = &FMediaPlateEditorStyle::Get().Get();
+						check(Style);
+
+						FUIAction Action_ConfigureComposite(
+							FExecuteAction::CreateLambda([](AMediaPlate* InMediaPlateActor)
+								{
+									if (!IsValid(InMediaPlateActor))
+									{
+										return;
+									}
+
+									const FScopedTransaction Transaction(LOCTEXT("ApplyOverlayCompositeMats", "Apply Overlay Composite Materials"));
+									InMediaPlateActor->Modify();
+
+									UMaterial* BasePassMaterial = LoadObject<UMaterial>(NULL, TEXT("/MediaPlate/M_MediaPlate_Masked"), NULL, LOAD_None, NULL);
+									InMediaPlateActor->ApplyMaterial(BasePassMaterial);
+
+									UMaterial* OverlayMaterial = LoadObject<UMaterial>(NULL, TEXT("/MediaPlate/M_MediaPlate_OverlayComp"), NULL, LOAD_None, NULL);
+									InMediaPlateActor->ApplyOverlayMaterial(OverlayMaterial);
+
+									UE::MediaPlate::Private::ApplyTranslucencyScreenPercentageCVar(1);
+								}
+								,MediaPlateActor)
+						);
+
+						FUIAction Action_ResetDefault(
+							FExecuteAction::CreateLambda([](AMediaPlate* InMediaPlateActor)
+								{
+									if (!IsValid(InMediaPlateActor))
+									{
+										return;
+									}
+
+									const FScopedTransaction Transaction(LOCTEXT("ResetDefaultMats", "Reset Default Materials"));
+									
+									InMediaPlateActor->Modify();
+									InMediaPlateActor->UseDefaultMaterial();
+
+									UE::MediaPlate::Private::ApplyTranslucencyScreenPercentageCVar(0);
+								}
+								, MediaPlateActor)
+						);
+
+						MenuBuilder.BeginSection("MediaPlate", LOCTEXT("MediaPlateHeading", "Media Plate"));
+						MenuBuilder.AddMenuEntry(
+							LOCTEXT("ApplyOverlayCompositeMats", "Apply Overlay Composite Materials"),
+							LOCTEXT("ApplyOverlayCompositeMats_Tooltip", "Setup the media plate for overlay-compositing to avoid TSR artifacts by replacing relevant materials."),
+							FSlateIcon(Style->GetStyleSetName(), "ClassIcon.MediaPlate"),
+							Action_ConfigureComposite
+						);
+						MenuBuilder.AddMenuEntry(
+							LOCTEXT("ResetDefaultMats", "Reset Default Materials"),
+							LOCTEXT("ResetDefaultMats_Tooltip", "Reverts the media plate to its default materials. Note that we also globally disable r.Translucency.ScreenPercentage.Basis if previously enabled."),
+							FSlateIcon(Style->GetStyleSetName(), "ClassIcon.MediaPlate"),
+							Action_ResetDefault
+						);
+						MenuBuilder.EndSection();
+					},
+					MediaPlateActor
+				)
+			);
+		}
+
+
+	}
+
+	return Extender.ToSharedRef();
+}
+
+void FMediaPlateEditorModule::RegisterContextMenuExtender()
+{
+	// Extend the level viewport context menu to add an option to copy the object path.
+	LevelViewportContextMenuRemoteControlExtender = FLevelEditorModule::FLevelViewportMenuExtender_SelectedActors::CreateRaw(this, &FMediaPlateEditorModule::ExtendLevelViewportContextMenuForMediaPlate);
+	
+	FLevelEditorModule& LevelEditorModule = FModuleManager::Get().LoadModuleChecked<FLevelEditorModule>("LevelEditor");
+	TArray<FLevelEditorModule::FLevelViewportMenuExtender_SelectedActors>& MenuExtenders = LevelEditorModule.GetAllLevelViewportContextMenuExtenders();
+	MenuExtenders.Add(LevelViewportContextMenuRemoteControlExtender);
+	MenuExtenderDelegateHandle = MenuExtenders.Last().GetHandle();
+}
+
+void FMediaPlateEditorModule::UnregisterContextMenuExtender()
+{
+	if (FModuleManager::Get().IsModuleLoaded("LevelEditor"))
+	{
+		FLevelEditorModule& LevelEditorModule = FModuleManager::LoadModuleChecked<FLevelEditorModule>("LevelEditor");
+		LevelEditorModule.GetAllLevelViewportContextMenuExtenders().RemoveAll(
+			[&](const FLevelEditorModule::FLevelViewportMenuExtender_SelectedActors& Delegate)
+			{
+				return Delegate.GetHandle() == MenuExtenderDelegateHandle;
+			});
+	}
+}
+
 IMPLEMENT_MODULE(FMediaPlateEditorModule, MediaPlateEditor)
 
 #undef LOCTEXT_NAMESPACE
diff --git a/Engine/Plugins/Media/MediaPlate/Source/MediaPlateEditor/Public/MediaPlateEditorModule.h b/Engine/Plugins/Media/MediaPlate/Source/MediaPlateEditor/Public/MediaPlateEditorModule.h
--- a/Engine/Plugins/Media/MediaPlate/Source/MediaPlateEditor/Public/MediaPlateEditorModule.h
+++ b/Engine/Plugins/Media/MediaPlate/Source/MediaPlateEditor/Public/MediaPlateEditorModule.h
@@ -4,6 +4,7 @@
 
 #include "Containers/Map.h"
 #include "CoreMinimal.h"
+#include "LevelEditor.h"
 #include "Logging/LogMacros.h"
 #include "Modules/ModuleManager.h"
 #include "TickableEditorObject.h"
@@ -64,6 +65,12 @@ private:
 	/** Stores what we have aded to the placement module. */
 	TArray<TOptional<FPlacementModeID>> PlaceActors;
 
+	/** Holds the context menu extender. */
+	FLevelEditorModule::FLevelViewportMenuExtender_SelectedActors LevelViewportContextMenuRemoteControlExtender;
+
+	/** Holds the menu extender delegate handle. */
+	FDelegateHandle MenuExtenderDelegateHandle;
+
 	/**
 	 * Gathers the Info on the Media Plate Place Actors Category.
 	 */
@@ -101,4 +108,19 @@ private:
 	 * Enables real time viewports so media plates can update in the editor.
 	 */
 	void ForceRealTimeViewports(const bool bEnable);
+
+	/**
+	 * Handle adding the menu extender for the actors.
+	 */
+	TSharedRef<FExtender> ExtendLevelViewportContextMenuForMediaPlate(const TSharedRef<FUICommandList> CommandList, TArray<AActor*> SelectedActors);
+
+	/**
+	 * Create all required callbacks for context menu's widget and button creation.
+	 */
+	void RegisterContextMenuExtender();
+
+	/**
+	 * Release context menu related resources such as handles.
+	 */
+	void UnregisterContextMenuExtender();
 };
