[Media IO] Media Plate UX Improvements

- Categories rearrangement
- Play/Pause buttons are now mutually exclusive
- Aspect Ratio and Letterbox Ratio widgets are now of type SSpinBox
- Sphere Horizontal and Vertical Arc properties have not been switched to SSpinBox for performance reasons (scrubbing would just freeze the editor momentarily)

Note related to UE 5.4 Feature Complete Schedule:
This CL originally included a scrubber widget placed above the playback control buttons. That widget has been removed in order to submit this CL in time for UE 5.4 feature complete.
The scrubber widget is in CL 30926672 which still needs to go under review.
It should be fairly easy to add once that is reviewed.

diff --git a/Engine/Plugins/Media/MediaPlate/Source/MediaPlate/Public/MediaPlateComponent.h b/Engine/Plugins/Media/MediaPlate/Source/MediaPlate/Public/MediaPlateComponent.h
--- a/Engine/Plugins/Media/MediaPlate/Source/MediaPlate/Public/MediaPlateComponent.h
+++ b/Engine/Plugins/Media/MediaPlate/Source/MediaPlate/Public/MediaPlateComponent.h
@@ -3,8 +3,6 @@
 #pragma once
 
 #include "Components/ActorComponent.h"
-#include "CoreMinimal.h"
-#include "MediaPlayerOptions.h"
 #include "MediaPlayerProxyInterface.h"
 #include "MediaSource.h"
 #include "MediaTextureTracker.h"
@@ -48,11 +46,11 @@ struct FMediaTextureResourceSettings
 	GENERATED_USTRUCT_BODY()
 
 	/** Enable mips generation */
-	UPROPERTY(EditAnywhere, Category = "Media Texture Resource Settings")
+	UPROPERTY(EditAnywhere, Category = "MediaTexture", meta = (DisplayName = "Enable RealTime Mips"))
 	bool bEnableGenMips = false;
 
 	/** Current number of mips to be generated as output */
-	UPROPERTY(EditAnywhere, Category = "Media Texture Resource Settings")
+	UPROPERTY(EditAnywhere, Category = "MediaTexture", meta = (DisplayName = "Mips Quantity"))
 	uint8 CurrentNumMips = 1;
 };
 
@@ -187,14 +185,15 @@ public:
 	float StartTime = 0.0f;
 
 	/** Holds the component to play sound. */
-	UPROPERTY(EditAnywhere, AdvancedDisplay, Category = MediaPlate)
+	UPROPERTY(EditAnywhere, Category = "Advanced", meta = (DisplayName = "Audio Component"))
 	TObjectPtr<UMediaSoundComponent> SoundComponent;
 
 	/** Holds the component for the mesh. */
-	UPROPERTY(EditAnywhere, AdvancedDisplay, Category = MediaPlate)
+	UPROPERTY(EditAnywhere, Category = "Advanced")
 	TObjectPtr<UStaticMeshComponent> StaticMeshComponent;
+
 	/** Holds the component for the mesh. */
-	UPROPERTY(EditAnywhere, AdvancedDisplay, Category = MediaPlate)
+	UPROPERTY(EditAnywhere, Category = "Advanced|Other")
 	TArray<TObjectPtr<UStaticMeshComponent>> Letterboxes;
 
 	/** What media playlist to play. */
@@ -206,7 +205,7 @@ public:
 	int32 PlaylistIndex = 0;
 
 	/** Override the default cache settings. */
-	UPROPERTY(EditAnywhere, AdvancedDisplay, Category = "MediaPlate")
+	UPROPERTY(EditAnywhere, Category = "Cache", meta = (DisplayName = "Cache", ShowOnlyInnerProperties))
 	FMediaSourceCacheSettings CacheSettings;
 
 	/** Set the arc size in degrees used for visible mips and tiles calculations, specific to the sphere. */
@@ -347,11 +346,11 @@ private:
 	bool bLoop = true;
 
 	/** Visible mips and tiles calculation mode for the supported mesh types in MediaPlate. (Player restart on change.) */
-	UPROPERTY(EditAnywhere, AdvancedDisplay, Category = "MediaPlate", meta = (AllowPrivateAccess = true))
+	UPROPERTY(EditAnywhere, Category = "EXR Tiles & Mips", meta = (DisplayName = "Visible Tiles & Mips Logic", AllowPrivateAccess = true))
 	EMediaTextureVisibleMipsTiles VisibleMipsTilesCalculations;
 
 	/** Media texture mip map bias shared between the (image sequence) loader and the media texture sampler. */
-	UPROPERTY(EditAnywhere, AdvancedDisplay, Category = "MediaPlate", meta = (AllowPrivateAccess = true, UIMin = "-16.0", UIMax = "15.99"))
+	UPROPERTY(EditAnywhere, Category = "EXR Tiles & Mips", meta = (DisplayName = "Mips Bias", AllowPrivateAccess = true, UIMin = "-16.0", UIMax = "15.99"))
 	float MipMapBias = 0.0f;
 
 	/** If true then set the aspect ratio automatically based on the media. */
@@ -359,15 +358,15 @@ private:
 	bool bIsAspectRatioAuto = true;
 
 	/** If true then enable the use of MipLevelToUpscale as defined below. */
-	UPROPERTY(EditAnywhere, AdvancedDisplay, Category = "MediaPlate", meta = (AllowPrivateAccess = true))
+	UPROPERTY(EditAnywhere, Category = "EXR Tiles & Mips", meta = (DisplayName = "Enable Mip Upscaling", AllowPrivateAccess = true))
 	bool bEnableMipMapUpscaling = false;
 
 	/* With exr playback, upscale into lower quality mips from this specified level. All levels including and above the specified value will be fully read. */
-	UPROPERTY(EditAnywhere, AdvancedDisplay, Category = "MediaPlate", meta = (EditCondition = "bEnableMipMapUpscaling", AllowPrivateAccess = true, UIMin = "0", UIMax = "16"))
+	UPROPERTY(EditAnywhere, Category = "EXR Tiles & Mips", meta = (DisplayName = "Upscale Mip Level", EditCondition = "bEnableMipMapUpscaling", EditConditionHides, AllowPrivateAccess = true, UIMin = "0", UIMax = "16"))
 	int32 MipLevelToUpscale = 16;
 
 	/** If true then Media Plate will attempt to load and upscale lower quality mips and display those at the poles (Sphere object only). */
-	UPROPERTY(EditAnywhere, AdvancedDisplay, Category = "MediaPlate", meta = (EditCondition = "VisibleMipsTilesCalculations == EMediaTextureVisibleMipsTiles::Sphere", AllowPrivateAccess = true))
+	UPROPERTY(EditAnywhere, Category = "EXR Tiles & Mips", meta = (DisplayName = "Adaptive Pole Mip Upscale", EditCondition = "VisibleMipsTilesCalculations == EMediaTextureVisibleMipsTiles::Sphere", EditConditionHides, AllowPrivateAccess = true))
 	bool bAdaptivePoleMipUpscaling = true;
 
 	/** If > 0, then this is the aspect ratio of our screen and 
@@ -397,7 +396,7 @@ private:
 	TArray<TObjectPtr<UMediaTexture>> MediaTextures;
 
 	/** Exposes Media Texture settings via Media Plate component. */
-	UPROPERTY(EditAnywhere, AdvancedDisplay, Category = "MediaPlate")
+	UPROPERTY(EditAnywhere, Category = "MediaTexture", meta = (ShowOnlyInnerProperties))
 	FMediaTextureResourceSettings MediaTextureSettings;
 
 	/** This component's media player */
diff --git a/Engine/Plugins/Media/MediaPlate/Source/MediaPlateEditor/Private/MediaPlateCustomization.cpp b/Engine/Plugins/Media/MediaPlate/Source/MediaPlateEditor/Private/MediaPlateCustomization.cpp
--- a/Engine/Plugins/Media/MediaPlate/Source/MediaPlateEditor/Private/MediaPlateCustomization.cpp
+++ b/Engine/Plugins/Media/MediaPlate/Source/MediaPlateEditor/Private/MediaPlateCustomization.cpp
@@ -8,7 +8,6 @@
 #include "DetailLayoutBuilder.h"
 #include "DetailWidgetRow.h"
 #include "Framework/MultiBox/MultiBoxBuilder.h"
-#include "IDetailGroup.h"
 #include "IMediaAssetsModule.h"
 #include "LevelEditorViewport.h"
 #include "MediaPlate.h"
@@ -61,15 +60,12 @@ void FMediaPlateCustomization::CustomizeDetails(IDetailLayoutBuilder& DetailBuil
 	// Get style.
 	const ISlateStyle* Style = &FMediaPlateEditorStyle::Get().Get();
 
-	// Hide the static mesh.
-	IDetailCategoryBuilder& StaticMeshCategory = DetailBuilder.EditCategory("StaticMesh");
-	StaticMeshCategory.SetCategoryVisibility(false);
+	CustomizeCategories(DetailBuilder);
 
 	IDetailCategoryBuilder& ControlCategory = DetailBuilder.EditCategory("Control");
-	IDetailCategoryBuilder& MediaPlateCategory = DetailBuilder.EditCategory("MediaPlate");
-	// The sort order number places these categories after the transform in the detail pane.
-	ControlCategory.SetSortOrder(2010);
-	MediaPlateCategory.SetSortOrder(2020);
+	IDetailCategoryBuilder& PlaylistCategory = DetailBuilder.EditCategory("Playlist");
+	IDetailCategoryBuilder& GeometryCategory = DetailBuilder.EditCategory("Geometry");
+	IDetailCategoryBuilder& MediaDetailsCategory = DetailBuilder.EditCategory("MediaDetails");
 
 	// Get objects we are editing.
 	TArray<TWeakObjectPtr<UObject>> Objects;
@@ -90,14 +86,10 @@ void FMediaPlateCustomization::CustomizeDetails(IDetailLayoutBuilder& DetailBuil
 	UpdateMediaPath();
 
 	// Add mesh customization.
-	AddMeshCustomization(MediaPlateCategory);
-	
-	// Create playlist group.
-	IDetailGroup& PlaylistGroup = MediaPlateCategory.AddGroup(TEXT("Playlist"),
-		LOCTEXT("Playlist", "Playlist"));
+	AddMeshCustomization(GeometryCategory);
 
 	// Add playlist.
-	PlaylistGroup.AddWidgetRow()
+	PlaylistCategory.AddCustomRow(FText::FromString("Playlist Asset"))
 		.NameContent()
 		[
 			SNew(STextBlock)
@@ -117,7 +109,7 @@ void FMediaPlateCustomization::CustomizeDetails(IDetailLayoutBuilder& DetailBuil
 	// Add media source type.
 	TAttribute<EVisibility> MediaSourceAssetVisibility(this, &FMediaPlateCustomization::ShouldShowMediaSourceAsset);
 	TAttribute<EVisibility> MediaSourceFileVisibility(this, &FMediaPlateCustomization::ShouldShowMediaSourceFile);
-	PlaylistGroup.AddWidgetRow()
+	PlaylistCategory.AddCustomRow(FText::FromString("First Item In Playlist"))
 		.NameContent()
 		[
 			SNew(STextBlock)
@@ -149,7 +141,7 @@ void FMediaPlateCustomization::CustomizeDetails(IDetailLayoutBuilder& DetailBuil
 		];
 
 	// Add media asset.
-	PlaylistGroup.AddWidgetRow()
+	PlaylistCategory.AddCustomRow(FText::FromString("Media Source Asset"))
 		.Visibility(MediaSourceAssetVisibility)
 		.NameContent()
 		[
@@ -168,7 +160,7 @@ void FMediaPlateCustomization::CustomizeDetails(IDetailLayoutBuilder& DetailBuil
 
 	// Add media path.
 	FString FileTypeFilter = TEXT("All files (*.*)|*.*");
-	PlaylistGroup.AddWidgetRow()
+	PlaylistCategory.AddCustomRow(FText::FromString("Media Path"))
 		.Visibility(MediaSourceFileVisibility)
 		.NameContent()
 		[
@@ -193,7 +185,6 @@ void FMediaPlateCustomization::CustomizeDetails(IDetailLayoutBuilder& DetailBuil
 	ControlCategory.AddCustomRow(LOCTEXT("MediaPlateControls", "MediaPlate Controls"))
 		[
 			SNew(SHorizontalBox)
-
 			// Rewind button.
 			+ SHorizontalBox::Slot()
 				.VAlign(VAlign_Center)
@@ -203,20 +194,13 @@ void FMediaPlateCustomization::CustomizeDetails(IDetailLayoutBuilder& DetailBuil
 						.VAlign(VAlign_Center)
 						.IsEnabled_Lambda([this]
 						{
-							for (TWeakObjectPtr<UMediaPlateComponent>& MediaPlatePtr : MediaPlatesList)
+							if (const UMediaPlayer* MediaPlayer = GetMediaPlayer())
 							{
-								UMediaPlateComponent* MediaPlate = MediaPlatePtr.Get();
-								if (MediaPlate != nullptr)
-								{
-									UMediaPlayer* MediaPlayer = MediaPlate->GetMediaPlayer();
-									if (MediaPlayer != nullptr)
-									{
-										return MediaPlayer->IsReady() &&
-											MediaPlayer->SupportsSeeking() &&
-											MediaPlayer->GetTime() > FTimespan::Zero();
-									}
-								}
+								return MediaPlayer->IsReady() &&
+									MediaPlayer->SupportsSeeking() &&
+									MediaPlayer->GetTime() > FTimespan::Zero();
 							}
+
 							return false;
 						})
 						.OnClicked_Lambda([this]() -> FReply
@@ -239,19 +223,13 @@ void FMediaPlateCustomization::CustomizeDetails(IDetailLayoutBuilder& DetailBuil
 				[
 					SNew(SButton)
 						.VAlign(VAlign_Center)
-						.IsEnabled_Lambda([this] {
-							for (TWeakObjectPtr<UMediaPlateComponent>& MediaPlatePtr : MediaPlatesList)
+						.IsEnabled_Lambda([this]
+						{
+							if (UMediaPlayer* MediaPlayer = GetMediaPlayer())
 							{
-								UMediaPlateComponent* MediaPlate = MediaPlatePtr.Get();
-								if (MediaPlate != nullptr)
-								{
-									UMediaPlayer* MediaPlayer = MediaPlate->GetMediaPlayer();
-									if (MediaPlayer != nullptr)
-									{
-										return MediaPlayer->IsReady() && MediaPlayer->SupportsRate(UMediaPlateComponent::GetReverseRate(MediaPlayer), false);
-									}
-								}
+								return MediaPlayer->IsReady() && MediaPlayer->SupportsRate(UMediaPlateComponent::GetReverseRate(MediaPlayer), false);
 							}
+
 							return false;
 						})
 						.OnClicked_Lambda([this]() -> FReply
@@ -276,20 +254,24 @@ void FMediaPlateCustomization::CustomizeDetails(IDetailLayoutBuilder& DetailBuil
 						.VAlign(VAlign_Center)
 						.IsEnabled_Lambda([this]
 						{
-							for (TWeakObjectPtr<UMediaPlateComponent>& MediaPlatePtr : MediaPlatesList)
+							// Is the player paused or fast forwarding/rewinding?
+							if (const UMediaPlayer* MediaPlayer = GetMediaPlayer())
 							{
-								UMediaPlateComponent* MediaPlate = MediaPlatePtr.Get();
-								if (MediaPlate != nullptr)
-								{
-									// Is the player paused or fast forwarding/rewinding?
-									UMediaPlayer* MediaPlayer = MediaPlate->GetMediaPlayer();
-									return ((MediaPlayer != nullptr) &&
-										(MediaPlayer->IsReady()) &&
-										((!MediaPlayer->IsPlaying() || (MediaPlayer->GetRate() != 1.0f))));
-								}
+								return MediaPlayer->IsReady()
+									&& (!MediaPlayer->IsPlaying() || (MediaPlayer->GetRate() != 1.0f));
 							}
+
 							return false;
 						})
+						.Visibility_Lambda([this]
+						{
+							if (const UMediaPlayer* MediaPlayer = GetMediaPlayer())
+							{
+								return MediaPlayer->IsPlaying() ? EVisibility::Collapsed : EVisibility::Visible;
+							}
+
+							return EVisibility::Visible;
+						})
 						.OnClicked_Lambda([this]() -> FReply
 						{
 							OnButtonEvent(EMediaPlateEventState::Play);
@@ -312,20 +294,23 @@ void FMediaPlateCustomization::CustomizeDetails(IDetailLayoutBuilder& DetailBuil
 						.VAlign(VAlign_Center)
 						.IsEnabled_Lambda([this]
 						{
-							for (TWeakObjectPtr<UMediaPlateComponent>& MediaPlatePtr : MediaPlatesList)
+							if (const UMediaPlayer* MediaPlayer = GetMediaPlayer())
 							{
-								UMediaPlateComponent* MediaPlate = MediaPlatePtr.Get();
-								if (MediaPlate != nullptr)
-								{
-									UMediaPlayer* MediaPlayer = MediaPlate->GetMediaPlayer();
-									if (MediaPlayer != nullptr)
-									{
-										return (MediaPlayer->CanPause()) && (!MediaPlayer->IsPaused());
-									}
-								}
+								return MediaPlayer->CanPause() && !MediaPlayer->IsPaused();
 							}
+
 							return false;
 						})
+						.Visibility_Lambda([this]
+						{
+							if (const UMediaPlayer* MediaPlayer = GetMediaPlayer())
+							{
+								const bool bIsVisible = MediaPlayer->CanPause() && !MediaPlayer->IsPaused();
+								return bIsVisible ? EVisibility::Visible : EVisibility::Collapsed;
+							}
+
+							return EVisibility::Collapsed;
+						})
 						.OnClicked_Lambda([this]() -> FReply
 						{
 							OnButtonEvent(EMediaPlateEventState::Pause);
@@ -346,19 +331,13 @@ void FMediaPlateCustomization::CustomizeDetails(IDetailLayoutBuilder& DetailBuil
 				[
 					SNew(SButton)
 					.VAlign(VAlign_Center)
-					.IsEnabled_Lambda([this] {
-						for (TWeakObjectPtr<UMediaPlateComponent>& MediaPlatePtr : MediaPlatesList)
+					.IsEnabled_Lambda([this]
+					{
+						if (UMediaPlayer* MediaPlayer = GetMediaPlayer())
 						{
-							UMediaPlateComponent* MediaPlate = MediaPlatePtr.Get();
-							if (MediaPlate != nullptr)
-							{
-								UMediaPlayer* MediaPlayer = MediaPlate->GetMediaPlayer();
-								if (MediaPlayer != nullptr)
-								{
-									return MediaPlayer->IsReady() && MediaPlayer->SupportsRate(UMediaPlateComponent::GetForwardRate(MediaPlayer), false);
-								}
-							}
+							return MediaPlayer->IsReady() && MediaPlayer->SupportsRate(UMediaPlateComponent::GetForwardRate(MediaPlayer), false);
 						}
+
 						return false;
 					})
 					.OnClicked_Lambda([this]() -> FReply
@@ -403,18 +382,11 @@ void FMediaPlateCustomization::CustomizeDetails(IDetailLayoutBuilder& DetailBuil
 						.VAlign(VAlign_Center)
 						.IsEnabled_Lambda([this]
 						{
-							for (TWeakObjectPtr<UMediaPlateComponent>& MediaPlatePtr : MediaPlatesList)
+							if (UMediaPlayer* MediaPlayer = GetMediaPlayer())
 							{
-								UMediaPlateComponent* MediaPlate = MediaPlatePtr.Get();
-								if (MediaPlate != nullptr)
-								{
-									UMediaPlayer* MediaPlayer = MediaPlate->GetMediaPlayer();
-									if (MediaPlayer != nullptr)
-									{
-										return !MediaPlayer->GetUrl().IsEmpty();
-									}
-								}
+								return !MediaPlayer->GetUrl().IsEmpty();
 							}
+
 							return false;
 						})
 						.OnClicked_Lambda([this]() -> FReply
@@ -435,10 +407,9 @@ void FMediaPlateCustomization::CustomizeDetails(IDetailLayoutBuilder& DetailBuil
 	// Add button to open the media plate editor.
 	if (bIsMediaPlateWindow == false)
 	{
-		MediaPlateCategory.AddCustomRow(LOCTEXT("OpenMediaPlate", "Open Media Plate"))
+		PlaylistCategory.AddCustomRow(LOCTEXT("OpenMediaPlate", "Open Media Plate"))
 			[
 				SNew(SHorizontalBox)
-
 				+ SHorizontalBox::Slot()
 					.FillWidth(1.f)
 					.Padding(0, 5, 10, 5)
@@ -466,11 +437,7 @@ void FMediaPlateCustomization::CustomizeDetails(IDetailLayoutBuilder& DetailBuil
 
 		if (FirstMediaPlate != nullptr)
 		{
-			// Create media details group.
-			IDetailGroup& MediaDetailsGroup = MediaPlateCategory.AddGroup(TEXT("MediaDetails"),
-				LOCTEXT("MediaDetails", "Media Details"));
-
-			MediaDetailsGroup.AddWidgetRow()
+			MediaDetailsCategory.AddCustomRow(FText::FromString(TEXT("Media Details")))
 			[
 				SNew(SMediaPlateEditorMediaDetails, *FirstMediaPlate)
 			];
@@ -478,37 +445,33 @@ void FMediaPlateCustomization::CustomizeDetails(IDetailLayoutBuilder& DetailBuil
 	}
 }
 
-void FMediaPlateCustomization::AddMeshCustomization(IDetailCategoryBuilder& MediaPlateCategory)
+void FMediaPlateCustomization::AddMeshCustomization(IDetailCategoryBuilder& InParentCategory)
 {
-	// Create detail group.
-	IDetailGroup& DetailGroup = MediaPlateCategory.AddGroup(TEXT("Mesh"),
-		LOCTEXT("Mesh", "Mesh"));
-
 	// Add radio buttons for mesh type.
-	DetailGroup.AddWidgetRow()
-		[
-			SNew(SSegmentedControl<EMediaTextureVisibleMipsTiles>)
-				.Value_Lambda([this]()
-				{
-					return MeshMode;
-				})
-				.OnValueChanged(this, &FMediaPlateCustomization::SetMeshMode)
-
-			+ SSegmentedControl<EMediaTextureVisibleMipsTiles>::Slot(EMediaTextureVisibleMipsTiles::Plane)
-				.Text(LOCTEXT("Plane", "Plane"))
-				.ToolTip(LOCTEXT("Plane_ToolTip",
-					"Select this if you want to use a standard plane for the mesh."))
-
-			+ SSegmentedControl<EMediaTextureVisibleMipsTiles>::Slot(EMediaTextureVisibleMipsTiles::Sphere)
-				.Text(LOCTEXT("Sphere", "Sphere"))
-				.ToolTip(LOCTEXT("Sphere_ToolTip",
-					"Select this if you want to use a spherical object for the mesh."))
-
-			+ SSegmentedControl<EMediaTextureVisibleMipsTiles>::Slot(EMediaTextureVisibleMipsTiles::None)
-				.Text(LOCTEXT("Custom", "Custom"))
-				.ToolTip(LOCTEXT("Custom_ToolTip",
-					"Select this if you want to provide your own mesh."))
-		];
+	InParentCategory.AddCustomRow(FText::FromString("Mesh Selection"))
+	[
+		SNew(SSegmentedControl<EMediaTextureVisibleMipsTiles>)
+			.Value_Lambda([this]()
+			{
+				return MeshMode;
+			})
+			.OnValueChanged(this, &FMediaPlateCustomization::SetMeshMode)
+
+		+ SSegmentedControl<EMediaTextureVisibleMipsTiles>::Slot(EMediaTextureVisibleMipsTiles::Plane)
+			.Text(LOCTEXT("Plane", "Plane"))
+			.ToolTip(LOCTEXT("Plane_ToolTip",
+				"Select this if you want to use a standard plane for the mesh."))
+
+		+ SSegmentedControl<EMediaTextureVisibleMipsTiles>::Slot(EMediaTextureVisibleMipsTiles::Sphere)
+			.Text(LOCTEXT("Sphere", "Sphere"))
+			.ToolTip(LOCTEXT("Sphere_ToolTip",
+				"Select this if you want to use a spherical object for the mesh."))
+
+		+ SSegmentedControl<EMediaTextureVisibleMipsTiles>::Slot(EMediaTextureVisibleMipsTiles::None)
+			.Text(LOCTEXT("Custom", "Custom"))
+			.ToolTip(LOCTEXT("Custom_ToolTip",
+				"Select this if you want to provide your own mesh."))
+	];
 
 	// Visibility attributes.
 	TAttribute<EVisibility> MeshCustomVisibility(this, &FMediaPlateCustomization::ShouldShowMeshCustomWidgets);
@@ -516,7 +479,7 @@ void FMediaPlateCustomization::AddMeshCustomization(IDetailCategoryBuilder& Medi
 	TAttribute<EVisibility> MeshSphereVisibility(this, &FMediaPlateCustomization::ShouldShowMeshSphereWidgets);
 
 	// Add aspect ratio.
-	DetailGroup.AddWidgetRow()
+	InParentCategory.AddCustomRow(FText::FromString("Mesh Selection"))
 		.Visibility(MeshPlaneVisibility)
 		.NameContent()
 		[
@@ -549,7 +512,7 @@ void FMediaPlateCustomization::AddMeshCustomization(IDetailCategoryBuilder& Medi
 			+ SHorizontalBox::Slot()
 				.AutoWidth()
 				[
-					SNew(SNumericEntryBox<float>)
+					SNew(SSpinBox<float>)
 						.Value(this, &FMediaPlateCustomization::GetAspectRatio)
 						.MinValue(0.0f)
 						.OnValueChanged(this, &FMediaPlateCustomization::SetAspectRatio)
@@ -557,7 +520,7 @@ void FMediaPlateCustomization::AddMeshCustomization(IDetailCategoryBuilder& Medi
 		];
 
 		// Add letterbox aspect ratio.
-		DetailGroup.AddWidgetRow()
+		InParentCategory.AddCustomRow(FText::FromString("Aspect Ratio"))
 			.Visibility(MeshPlaneVisibility)
 			.NameContent()
 			[
@@ -591,7 +554,7 @@ void FMediaPlateCustomization::AddMeshCustomization(IDetailCategoryBuilder& Medi
 		+ SHorizontalBox::Slot()
 			.AutoWidth()
 			[
-				SNew(SNumericEntryBox<float>)
+				SNew(SSpinBox<float>)
 					.Value(this, &FMediaPlateCustomization::GetLetterboxAspectRatio)
 					.MinValue(0.0f)
 					.OnValueChanged(this, &FMediaPlateCustomization::SetLetterboxAspectRatio)
@@ -599,7 +562,7 @@ void FMediaPlateCustomization::AddMeshCustomization(IDetailCategoryBuilder& Medi
 		];
 
 	// Add auto aspect ratio.
-	DetailGroup.AddWidgetRow()
+	InParentCategory.AddCustomRow(FText::FromString("Aspect Ratio"))
 		.Visibility(MeshPlaneVisibility)
 		.NameContent()
 		[
@@ -617,7 +580,7 @@ void FMediaPlateCustomization::AddMeshCustomization(IDetailCategoryBuilder& Medi
 		];
 
 	// Add sphere horizontal arc.
-	DetailGroup.AddWidgetRow()
+	InParentCategory.AddCustomRow(FText::FromString("Horizontal Arc"))
 		.Visibility(MeshSphereVisibility)
 		.NameContent()
 		[
@@ -635,7 +598,7 @@ void FMediaPlateCustomization::AddMeshCustomization(IDetailCategoryBuilder& Medi
 		];
 
 	// Add sphere vertical arc.
-	DetailGroup.AddWidgetRow()
+	InParentCategory.AddCustomRow(FText::FromString("Vertical Arc"))
 		.Visibility(MeshSphereVisibility)
 		.NameContent()
 		[
@@ -653,7 +616,7 @@ void FMediaPlateCustomization::AddMeshCustomization(IDetailCategoryBuilder& Medi
 		];
 
 	// Add static mesh.
-	DetailGroup.AddWidgetRow()
+	InParentCategory.AddCustomRow(FText::FromString("Static Mesh"))
 		.Visibility(MeshCustomVisibility)
 		.NameContent()
 		[
@@ -670,7 +633,6 @@ void FMediaPlateCustomization::AddMeshCustomization(IDetailCategoryBuilder& Medi
 				.ObjectPath(this, &FMediaPlateCustomization::GetStaticMeshPath)
 				.OnObjectChanged(this, &FMediaPlateCustomization::OnStaticMeshChanged)
 		];
-	
 }
 
 EVisibility FMediaPlateCustomization::ShouldShowMeshCustomWidgets() const
@@ -821,8 +783,7 @@ void FMediaPlateCustomization::SetAspectRatio(float AspectRatio)
 	}
 }
 
-
-TOptional<float> FMediaPlateCustomization::GetAspectRatio() const
+float FMediaPlateCustomization::GetAspectRatio() const
 {
 	// Loop through our objects.
 	for (const TWeakObjectPtr<UMediaPlateComponent>& MediaPlatePtr : MediaPlatesList)
@@ -834,7 +795,7 @@ TOptional<float> FMediaPlateCustomization::GetAspectRatio() const
 		}
 	}
 
-	return TOptional<float>();
+	return 1.0f;
 }
 
 void FMediaPlateCustomization::SetLetterboxAspectRatio(float AspectRatio)
@@ -856,7 +817,7 @@ void FMediaPlateCustomization::SetLetterboxAspectRatio(float AspectRatio)
 }
 
 
-TOptional<float> FMediaPlateCustomization::GetLetterboxAspectRatio() const
+float FMediaPlateCustomization::GetLetterboxAspectRatio() const
 {
 	for (const TWeakObjectPtr<UMediaPlateComponent>& MediaPlatePtr : MediaPlatesList)
 	{
@@ -867,7 +828,7 @@ TOptional<float> FMediaPlateCustomization::GetLetterboxAspectRatio() const
 		}
 	}
 
-	return TOptional<float>();
+	return 0.0f;
 }
 
 void FMediaPlateCustomization::SetMeshHorizontalRange(float HorizontalRange)
@@ -1316,5 +1277,124 @@ void FMediaPlateCustomization::StopMediaPlates()
 	OnButtonEvent(EMediaPlateEventState::Close);
 }
 
+UMediaPlayer* FMediaPlateCustomization::GetMediaPlayer() const
+{
+	for (const TWeakObjectPtr<UMediaPlateComponent>& MediaPlatePtr : MediaPlatesList)
+	{
+		if (UMediaPlateComponent* MediaPlate = MediaPlatePtr.Get())
+		{
+			if (UMediaPlayer* MediaPlayer = MediaPlate->GetMediaPlayer())
+			{
+				return MediaPlayer;
+			}
+		}
+	}
+
+	return nullptr;
+}
+
+void FMediaPlateCustomization::CustomizeCategories(IDetailLayoutBuilder& InDetailBuilder)
+{
+	static const FName PropertyEditor("PropertyEditor");
+	FPropertyEditorModule& PropertyModule = FModuleManager::GetModuleChecked<FPropertyEditorModule>(PropertyEditor);
+
+	// Rearrange Categories
+
+	const FName MediaPlateComponentName = UMediaPlateComponent::StaticClass()->GetFName();
+
+	const TSharedRef<FPropertySection> GeneralSection = PropertyModule.FindOrCreateSection(MediaPlateComponentName, TEXT("General"), LOCTEXT("General", "General"));
+	GeneralSection->AddCategory(TEXT("Control"));
+	GeneralSection->AddCategory(TEXT("Geometry"));
+	GeneralSection->AddCategory(TEXT("Playlist"));
+	GeneralSection->AddCategory(TEXT("MediaDetails"));
+	GeneralSection->AddCategory(TEXT("MediaTexture"));
+	GeneralSection->AddCategory(TEXT("Materials"));
+	GeneralSection->AddCategory(TEXT("EXR Tiles & Mips"));
+	GeneralSection->AddCategory(TEXT("Media Cache"));
+	GeneralSection->AddCategory(TEXT("Advanced"));
+
+	const TSharedRef<FPropertySection> MediaSection = PropertyModule.FindOrCreateSection(MediaPlateComponentName, TEXT("Media"), LOCTEXT("Media", "Media"));
+	MediaSection->AddCategory(TEXT("Playlist"));
+	MediaSection->AddCategory(TEXT("MediaDetails"));
+	MediaSection->AddCategory(TEXT("Media Cache"));
+
+	const TSharedRef<FPropertySection> EXRSection = PropertyModule.FindOrCreateSection(MediaPlateComponentName, TEXT("EXR"), LOCTEXT("EXR", "EXR"));
+	EXRSection->AddCategory(TEXT("MediaDetails"));
+	EXRSection->AddCategory(TEXT("EXR Tiles & Mips"));
+	EXRSection->AddCategory(TEXT("Media Cache"));
+
+	const TSharedRef<FPropertySection> RenderingSection = PropertyModule.FindOrCreateSection(MediaPlateComponentName, TEXT("Rendering"), LOCTEXT("Rendering", "Rendering"));
+	RenderingSection->AddCategory(TEXT("Geometry"));
+	RenderingSection->AddCategory(TEXT("Materials"));
+	RenderingSection->AddCategory(TEXT("MediaTexture"));
+	RenderingSection->AddCategory(TEXT("Mobility"));
+	RenderingSection->AddCategory(TEXT("Transform"));
+	RenderingSection->AddCategory(TEXT("TransformCommon"));
+	RenderingSection->RemoveCategory(TEXT("Lighting"));
+	RenderingSection->AddCategory(TEXT("MediaTexture"));
+	RenderingSection->RemoveCategory(TEXT("MaterialParameters"));
+	RenderingSection->RemoveCategory(TEXT("Mobile"));
+	RenderingSection->RemoveCategory(TEXT("RayTracing"));
+	RenderingSection->RemoveCategory(TEXT("Rendering"));
+	RenderingSection->RemoveCategory(TEXT("TextureStreaming"));
+	RenderingSection->RemoveCategory(TEXT("VirtualTexture"));
+
+	// Hide unwanted Categories
+
+	const FName MediaPlateName = AMediaPlate::StaticClass()->GetFName();
+
+	const TSharedRef<FPropertySection> MediaPlateMiscSection = PropertyModule.FindOrCreateSection(MediaPlateName, TEXT("Misc"), LOCTEXT("Misc", "Misc"));
+	MediaPlateMiscSection->RemoveCategory("AssetUserData");
+	MediaPlateMiscSection->RemoveCategory("Cooking");
+	MediaPlateMiscSection->RemoveCategory("Input");
+	MediaPlateMiscSection->RemoveCategory("Navigation");
+	MediaPlateMiscSection->RemoveCategory("Replication");
+	MediaPlateMiscSection->RemoveCategory("Tags");
+
+	const TSharedRef<FPropertySection> MediaPlateStreamingSection = PropertyModule.FindOrCreateSection(MediaPlateName, TEXT("Streaming"), LOCTEXT("Streaming", "Streaming"));
+	MediaPlateStreamingSection->RemoveCategory("Data Layers");
+	MediaPlateStreamingSection->RemoveCategory("HLOD");
+	MediaPlateStreamingSection->RemoveCategory("World Partition");
+
+	const TSharedRef<FPropertySection> MediaPlateLODSection = PropertyModule.FindOrCreateSection(MediaPlateName, TEXT("LOD"), LOCTEXT("LOD", "LOD"));
+	MediaPlateLODSection->RemoveCategory("HLOD");
+	MediaPlateLODSection->RemoveCategory("LOD");
+
+	const TSharedRef<FPropertySection> MediaPlatePhysicsSection = PropertyModule.FindOrCreateSection(MediaPlateName, TEXT("Physics"), LOCTEXT("Physics", "Physics"));
+	MediaPlatePhysicsSection->RemoveCategory("Collision");
+	MediaPlatePhysicsSection->RemoveCategory("Physics");
+
+	// Hide the static mesh.
+	IDetailCategoryBuilder& StaticMeshCategory = InDetailBuilder.EditCategory("StaticMesh");
+	StaticMeshCategory.SetCategoryVisibility(false);
+
+	IDetailCategoryBuilder& ControlCategory = InDetailBuilder.EditCategory("Control");
+	IDetailCategoryBuilder& MediaDetailsCategory = InDetailBuilder.EditCategory("MediaDetails");
+	IDetailCategoryBuilder& PlaylistCategory = InDetailBuilder.EditCategory("Playlist");
+	IDetailCategoryBuilder& GeometryCategory = InDetailBuilder.EditCategory("Geometry");
+	IDetailCategoryBuilder& MediaTextureCategory = InDetailBuilder.EditCategory("MediaTexture");
+	IDetailCategoryBuilder& MaterialsCategory = InDetailBuilder.EditCategory("Materials");
+	IDetailCategoryBuilder& TilesMipsCategory = InDetailBuilder.EditCategory("EXR Tiles & Mips");
+	IDetailCategoryBuilder& MediaCacheCategory = InDetailBuilder.EditCategory("Media Cache");
 
+	// Rename Media Cache category and look ahead property
+	MediaCacheCategory.SetDisplayName(FText::FromString(TEXT("Cache")));
+
+	const TSharedRef<IPropertyHandle> CacheSettingsProperty = InDetailBuilder.GetProperty(GET_MEMBER_NAME_CHECKED(UMediaPlateComponent, CacheSettings));
+	if (const auto LookAheadTimeProperty = CacheSettingsProperty->GetChildHandle(TEXT("TimeToLookAhead")))
+	{
+		LookAheadTimeProperty->SetPropertyDisplayName(FText::FromString("Look Ahead Time"));
+	}
+
+	// Start from a Priority value which places these categories after the Transform one
+	uint32 Priority = 2010;
+	ControlCategory.SetSortOrder(Priority++);
+	GeometryCategory.SetSortOrder(Priority++);
+	PlaylistCategory.SetSortOrder(Priority++);
+	MediaDetailsCategory.SetSortOrder(Priority++);
+	MediaTextureCategory.SetSortOrder(Priority++);
+	MaterialsCategory.SetSortOrder(Priority++);
+	TilesMipsCategory.SetSortOrder(Priority++);
+	MediaCacheCategory.SetSortOrder(Priority);
+}
 #undef LOCTEXT_NAMESPACE
diff --git a/Engine/Plugins/Media/MediaPlate/Source/MediaPlateEditor/Private/MediaPlateCustomization.h b/Engine/Plugins/Media/MediaPlate/Source/MediaPlateEditor/Private/MediaPlateCustomization.h
--- a/Engine/Plugins/Media/MediaPlate/Source/MediaPlateEditor/Private/MediaPlateCustomization.h
+++ b/Engine/Plugins/Media/MediaPlate/Source/MediaPlateEditor/Private/MediaPlateCustomization.h
@@ -2,7 +2,6 @@
 
 #pragma once
 
-#include "CoreMinimal.h"
 #include "Editor.h"
 #include "IDetailCustomization.h"
 #include "Input/Reply.h"
@@ -30,7 +29,6 @@ public:
 	~FMediaPlateCustomization();
 
 	//~ IDetailCustomization interface
-
 	virtual void CustomizeDetails(IDetailLayoutBuilder& DetailBuilder) override;
 
 	/**
@@ -44,6 +42,7 @@ public:
 	}
 
 private:
+	static void CustomizeCategories(IDetailLayoutBuilder& InDetailBuilder);
 
 	/** Property change delegate used for static mesh material changes. */
 	FDelegateHandle PropertyChangeDelegate;
@@ -60,11 +59,11 @@ private:
 
 	/** True if the first media source in the playlist is an external asset. */
 	bool bIsMediaSourceAsset;
-
+	
 	/**
 	 * Adds widgets for editing the mesh.
 	 */
-	void AddMeshCustomization(IDetailCategoryBuilder& MediaPlateCategory);
+	void AddMeshCustomization(IDetailCategoryBuilder& InParentCategory);
 
 	/**
 	 * Controls visibility for widgets for custom meshes.
@@ -124,8 +123,8 @@ private:
 
 	/**
 	 * Call this to get the aspect ratio.
-	 */
-	TOptional<float> GetAspectRatio() const;
+	 */	
+	float GetAspectRatio() const;
 
 	/**
 	 * Call this to set the aspect ratio.
@@ -135,7 +134,7 @@ private:
 	/**
 	 * Call this to get the aspect ratio.
 	 */
-	TOptional<float> GetLetterboxAspectRatio() const;
+	float GetLetterboxAspectRatio() const;
 
 	/**
 	 * Call this to set the horizontal range of the mesh.
@@ -241,5 +240,7 @@ private:
 	 * Call this to stop all playback.
 	 */
 	void StopMediaPlates();
+
+	UMediaPlayer* GetMediaPlayer() const;
 };
 
diff --git a/Engine/Plugins/Media/MediaPlate/Source/MediaPlateEditor/Private/Widgets/SMediaPlateEditorMediaDetails.cpp b/Engine/Plugins/Media/MediaPlate/Source/MediaPlateEditor/Private/Widgets/SMediaPlateEditorMediaDetails.cpp
--- a/Engine/Plugins/Media/MediaPlate/Source/MediaPlateEditor/Private/Widgets/SMediaPlateEditorMediaDetails.cpp
+++ b/Engine/Plugins/Media/MediaPlate/Source/MediaPlateEditor/Private/Widgets/SMediaPlateEditorMediaDetails.cpp
@@ -20,6 +20,9 @@ void SMediaPlateEditorMediaDetails::Construct(const FArguments& InArgs,
 {
 	MediaPlate = &InMediaPlate;
 
+	const ISlateStyle& SlateStyle = FAppStyle::Get();
+	const FName StyleName = TEXT("SmallText");
+
 	ChildSlot
 		[
 			SNew(SScrollBox)
@@ -41,6 +44,7 @@ void SMediaPlateEditorMediaDetails::Construct(const FArguments& InArgs,
 								.Padding(4.0f)
 								[
 									SAssignNew(MediaPlayerName, STextBlock)
+									.TextStyle(SlateStyle, StyleName)
 								]
 
 							// Resolution.
@@ -50,6 +54,7 @@ void SMediaPlateEditorMediaDetails::Construct(const FArguments& InArgs,
 								.Padding(4.0f)
 								[
 									SAssignNew(ResolutionText, STextBlock)
+									.TextStyle(SlateStyle, StyleName)
 								]
 
 							// Frame rate.
@@ -59,6 +64,7 @@ void SMediaPlateEditorMediaDetails::Construct(const FArguments& InArgs,
 								.Padding(4.0f)
 								[
 									SAssignNew(FrameRateText, STextBlock)
+									.TextStyle(SlateStyle, StyleName)
 								]
 
 							// Resource size.
@@ -68,6 +74,7 @@ void SMediaPlateEditorMediaDetails::Construct(const FArguments& InArgs,
 								.Padding(4.0f)
 								[
 									SAssignNew(ResourceSizeText, STextBlock)
+									.TextStyle(SlateStyle, StyleName)
 								]
 
 							// Method.
@@ -77,6 +84,7 @@ void SMediaPlateEditorMediaDetails::Construct(const FArguments& InArgs,
 								.Padding(4.0f)
 								[
 									SAssignNew(MethodText, STextBlock)
+									.TextStyle(SlateStyle, StyleName)
 								]
 						]
 
@@ -92,6 +100,7 @@ void SMediaPlateEditorMediaDetails::Construct(const FArguments& InArgs,
 								.Padding(4.0f)
 								[
 									SAssignNew(FormatText, STextBlock)
+									.TextStyle(SlateStyle, StyleName)
 								]
 
 							// LOD bias.
@@ -101,6 +110,7 @@ void SMediaPlateEditorMediaDetails::Construct(const FArguments& InArgs,
 								.Padding(4.0f)
 								[
 									SAssignNew(LODBiasText, STextBlock)
+									.TextStyle(SlateStyle, StyleName)
 								]
 
 							// Num mips.
@@ -110,6 +120,7 @@ void SMediaPlateEditorMediaDetails::Construct(const FArguments& InArgs,
 								.Padding(4.0f)
 								[
 									SAssignNew(NumMipsText, STextBlock)
+									.TextStyle(SlateStyle, StyleName)
 								]
 
 							// Num tiles.
@@ -119,6 +130,7 @@ void SMediaPlateEditorMediaDetails::Construct(const FArguments& InArgs,
 								.Padding(4.0f)
 								[
 									SAssignNew(NumTilesText, STextBlock)
+									.TextStyle(SlateStyle, StyleName)
 								]
 						]
 				]
@@ -190,9 +202,9 @@ void SMediaPlateEditorMediaDetails::UpdateDetails()
 	LODBiasText->SetText(FText::Format(LOCTEXT("LODBias", "Combined LOD Bias: {0}"),
 		FText::AsNumber(LODBias)));
 	MethodText->SetText(FText::Format(LOCTEXT("Method", "Method: {0}"), Method));
-	NumMipsText->SetText(FText::Format(LOCTEXT("NumberOfMips", "Number Of Mips: {0}"),
+	NumMipsText->SetText(FText::Format(LOCTEXT("NumberOfMips", "Mips: {0}"),
 		FText::AsNumber(NumMips)));
-	NumTilesText->SetText(FText::Format(LOCTEXT("NumberOfTiles", "Number Of Tiles: {0}"),
+	NumTilesText->SetText(FText::Format(LOCTEXT("NumberOfTiles", "Tiles: {0}"),
 		FText::AsNumber(NumTotalTiles)));
 	ResolutionText->SetText(FText::Format(LOCTEXT("Resolution", "Resolution: {0}x{1}"),
 		FText::AsNumber(SurfaceWidth), FText::AsNumber(SurfaceHeight)));
