MediaPlate: Fix transform drag freezing the editor.

diff --git a/Engine/Plugins/Media/MediaPlate/Source/MediaPlate/Private/MediaPlate.cpp b/Engine/Plugins/Media/MediaPlate/Source/MediaPlate/Private/MediaPlate.cpp
--- a/Engine/Plugins/Media/MediaPlate/Source/MediaPlate/Private/MediaPlate.cpp
+++ b/Engine/Plugins/Media/MediaPlate/Source/MediaPlate/Private/MediaPlate.cpp
@@ -352,7 +352,7 @@ void AMediaPlate::OnPostSaveWorld(UWorld* InWorld, FObjectPostSaveContext Object
 
 void AMediaPlate::AddAssetUserData()
 {
-	if (StaticMeshComponent != nullptr)
+	if (StaticMeshComponent != nullptr && !StaticMeshComponent->HasAssetUserDataOfClass(UMediaPlateAssetUserData::StaticClass()))
 	{
 		UMediaPlateAssetUserData* AssetUserData = NewObject<UMediaPlateAssetUserData>(GetTransientPackage());
 		AssetUserData->OnPostEditChangeOwner.BindUObject(this, &AMediaPlate::ApplyCurrentMaterial);
