MediaPlate: Minor tweaks to overlay technique.

* Added tooltip caveat around media with translucency.
* Apply missing cvar at runtime if still needed.

#rnx

diff --git a/Engine/Plugins/Media/MediaPlate/Source/MediaPlateEditor/Private/MediaPlateEditorModule.cpp b/Engine/Plugins/Media/MediaPlate/Source/MediaPlateEditor/Private/MediaPlateEditorModule.cpp
--- a/Engine/Plugins/Media/MediaPlate/Source/MediaPlateEditor/Private/MediaPlateEditorModule.cpp
+++ b/Engine/Plugins/Media/MediaPlate/Source/MediaPlateEditor/Private/MediaPlateEditorModule.cpp
@@ -378,7 +378,7 @@ TSharedRef<FExtender> FMediaPlateEditorModule::ExtendLevelViewportContextMenuFor
 						MenuBuilder.BeginSection("MediaPlate", LOCTEXT("MediaPlateHeading", "Media Plate"));
 						MenuBuilder.AddMenuEntry(
 							LOCTEXT("ApplyOverlayCompositeMats", "Apply Overlay Composite Materials"),
-							LOCTEXT("ApplyOverlayCompositeMats_Tooltip", "Setup the media plate for overlay-compositing to avoid TSR artifacts by replacing relevant materials."),
+							LOCTEXT("ApplyOverlayCompositeMats_Tooltip", "Setup the media plate for overlay-compositing to avoid TSR artifacts by replacing relevant materials. This technique is only effective on opaque media currently."),
 							FSlateIcon(Style->GetStyleSetName(), "ClassIcon.MediaPlate"),
 							Action_ConfigureComposite
 						);
