MediaPlate: Fixed playback state issue for players whose actual state is not immediately reflected in IsPlaying() and IsPaused() after calling Play() or Pause()

#rnx

diff --git a/Engine/Plugins/Media/MediaPlate/Source/MediaPlate/Private/MediaPlateComponent.cpp b/Engine/Plugins/Media/MediaPlate/Source/MediaPlate/Private/MediaPlateComponent.cpp
--- a/Engine/Plugins/Media/MediaPlate/Source/MediaPlate/Private/MediaPlateComponent.cpp
+++ b/Engine/Plugins/Media/MediaPlate/Source/MediaPlate/Private/MediaPlateComponent.cpp
@@ -34,7 +34,7 @@ namespace UE::MediaPlateComponent
 	};
 	ENUM_CLASS_FLAGS(ESetUpTexturesFlags);
 
-	// Runs through media textures and sets Media Plate settings corresponding to Media Texture. 
+	// Runs through media textures and sets Media Plate settings corresponding to Media Texture.
 	void ApplyMediaTextureMipGenProperties(const FMediaTextureResourceSettings MediaTextureSettings, const TArray<TObjectPtr<UMediaTexture>>& MediaTextures)
 	{
 		for (TObjectPtr<UMediaTexture> MediaTexture : MediaTextures)
@@ -162,6 +162,8 @@ void UMediaPlateComponent::OnRegister()
 	}
 	MediaPlayer->OnMediaOpened.AddUniqueDynamic(this, &UMediaPlateComponent::OnMediaOpened);
 	MediaPlayer->OnEndReached.AddUniqueDynamic(this, &UMediaPlateComponent::OnMediaEnd);
+	MediaPlayer->OnPlaybackResumed.AddUniqueDynamic(this, &UMediaPlateComponent::OnMediaResumed);
+	MediaPlayer->OnPlaybackSuspended.AddUniqueDynamic(this, &UMediaPlateComponent::OnMediaSuspended);
 
 	// Set up media texture.
 	SetUpTextures(UE::MediaPlateComponent::ESetUpTexturesFlags::AllowSetPlayer);
@@ -218,19 +220,47 @@ void UMediaPlateComponent::TickComponent(float DeltaTime, enum ELevelTick TickTy
 
 	if (MediaPlayer != nullptr)
 	{
-		if ((CurrentRate != 0.0f) || (bWantsToPlayWhenVisible))
+		// Pending rate changes?
+		if (IntendedPlaybackState != ActualPlaybackState)
 		{
-			bool bIsVisible = IsVisible();
-			if (bIsVisible)
+			if (IntendedPlaybackState != PendingPlaybackState)
 			{
-				ResumeWhenVisible();
+				if (IntendedPlaybackState == EPlaybackState::Resume)
+				{
+					FTimespan PlayTime = GetResumeTime();
+					MediaPlayer->Seek(PlayTime);
+					IntendedPlaybackState = EPlaybackState::Playing;
+				}
+				PendingPlaybackState = IntendedPlaybackState;
+				if (IntendedPlaybackState == EPlaybackState::Playing)
+				{
+					MediaPlayer->Play();
+				}
+				else
+				{
+					MediaPlayer->Pause();
+				}
 			}
-			else
+		}
+		else
+		{
+			PendingPlaybackState = EPlaybackState::Unset;
+		}
+
+		// Perform visibility check only when not currently transitioning.
+		if (PendingPlaybackState == EPlaybackState::Unset)
+		{
+			if ((CurrentRate != 0.0f) || bWantsToPlayWhenVisible)
 			{
-				if (MediaPlayer->IsPlaying())
+				bool bIsVisible = IsVisible();
+				if (bIsVisible)
 				{
-					MediaPlayer->Pause();
-					TimeWhenPlaybackPaused = static_cast<float>(FApp::GetGameTime());
+					ResumeWhenVisible();
+				}
+				else if (ActualPlaybackState == EPlaybackState::Playing)
+				{
+					IntendedPlaybackState = EPlaybackState::Paused;
+					TimeWhenPlaybackPaused = FApp::GetGameTime();
 				}
 			}
 		}
@@ -271,6 +301,11 @@ void UMediaPlateComponent::Open()
 {
 	bIsMediaPlatePlaying = true;
 	CurrentRate = bPlayOnOpen ? 1.0f : 0.0f;
+	IntendedPlaybackState = bPlayOnOpen ? EPlaybackState::Playing : EPlaybackState::Paused;
+	PendingPlaybackState = EPlaybackState::Unset;
+	ActualPlaybackState = EPlaybackState::Paused;
+	TimeWhenPlaybackPaused = -1.0;
+
 	PlaylistIndex = 0;
 	SetNormalMode(true);
 
@@ -296,7 +331,7 @@ void UMediaPlateComponent::Open()
 	else
 	{
 		bWantsToPlayWhenVisible = true;
-		TimeWhenPlaybackPaused = static_cast<float>(FApp::GetGameTime());
+		TimeWhenPlaybackPaused = FApp::GetGameTime();
 	}
 
 	UpdateTicking();
@@ -325,19 +360,13 @@ bool UMediaPlateComponent::Next()
 
 void UMediaPlateComponent::Play()
 {
-	if (MediaPlayer != nullptr)
-	{
-		MediaPlayer->Play();
-	}
+	IntendedPlaybackState = EPlaybackState::Playing;
 	CurrentRate = 1.0f;
 }
 
 void UMediaPlateComponent::Pause()
 {
-	if (MediaPlayer != nullptr)
-	{
-		MediaPlayer->Pause();
-	}
+	IntendedPlaybackState = EPlaybackState::Paused;
 	CurrentRate = 0.0f;
 }
 
@@ -492,7 +521,7 @@ bool UMediaPlateComponent::PlayMediaSource(UMediaSource* InMediaSource, bool bIn
 	{
 		// Set cache settings.
 		InMediaSource->SetCacheSettings(CacheSettings);
-		
+
 		// Set media options.
 		if (MediaPlayer != nullptr)
 		{
@@ -645,14 +674,8 @@ float UMediaPlateComponent::GetProxyRate() const
 bool UMediaPlateComponent::SetProxyRate(float Rate)
 {
 	CurrentRate = Rate;
-
-	bool bSuccess = true;
-	if (MediaPlayer != nullptr)
-	{
-		bSuccess = MediaPlayer->SetRate(Rate);
-	}
-
-	return bSuccess;
+	IntendedPlaybackState = Rate == 0.0f ? EPlaybackState::Paused : EPlaybackState::Playing;
+	return MediaPlayer ? MediaPlayer->SetRate(Rate) : true;
 }
 
 bool UMediaPlateComponent::IsExternalControlAllowed()
@@ -824,7 +847,7 @@ void UMediaPlateComponent::RestartPlayer()
 {
 	if (MediaPlayer != nullptr)
 	{
-		if (MediaPlayer->IsPlaying())
+		if (IntendedPlaybackState == EPlaybackState::Playing)
 		{
 			MediaPlayer->Close();
 			Open();
@@ -860,21 +883,19 @@ void UMediaPlateComponent::ResumeWhenVisible()
 {
 	if (MediaPlayer != nullptr)
 	{
-		if (MediaPlayer->IsPaused())
+		if (ActualPlaybackState == EPlaybackState::Paused)
 		{
 			// Should we be playing?
-			if (CurrentRate != 0.0f)
+			if (CurrentRate != 0.0f && PendingPlaybackState == EPlaybackState::Unset)
 			{
-				FTimespan PlayTime = GetResumeTime();
-				MediaPlayer->Seek(PlayTime);
-				MediaPlayer->Play();
+				IntendedPlaybackState = EPlaybackState::Resume;
 			}
 		}
 		else if (bWantsToPlayWhenVisible)
 		{
 			if ((bResumeWhenOpened == false) &&
 				(MediaPlayer->IsPreparing() == false) &&
-				(MediaPlayer->IsPlaying() == false))
+				(ActualPlaybackState == EPlaybackState::Paused))
 			{
 				bResumeWhenOpened = true;
 				bWantsToPlayWhenVisible = false;
@@ -890,27 +911,30 @@ FTimespan UMediaPlateComponent::GetResumeTime()
 	if (MediaPlayer != nullptr)
 	{
 		PlayerTime = MediaPlayer->GetTime();
-		float CurrentTime = static_cast<float>(FApp::GetGameTime());
-		float ElapsedTime = CurrentTime - TimeWhenPlaybackPaused;
-		PlayerTime += FTimespan::FromSeconds(ElapsedTime);
-		
-		// Are we over the length of the media?
-		FTimespan MediaDuration = MediaPlayer->GetDuration();
-		if ((PlayerTime > MediaDuration) && (MediaDuration > FTimespan::Zero()))
+		if (TimeWhenPlaybackPaused > 0.0)
 		{
-			bool bIsPlaylist = (MediaPlaylist != nullptr) && (MediaPlaylist->Num() > 1);
-			if ((bLoop) && (bIsPlaylist == false))
-			{
-				PlayerTime %= MediaDuration;
-			}
-			else
+			double CurrentTime = FApp::GetGameTime();
+			double ElapsedTime = CurrentTime - TimeWhenPlaybackPaused;
+			PlayerTime += FTimespan::FromSeconds(ElapsedTime);
+
+			// Are we over the length of the media?
+			FTimespan MediaDuration = MediaPlayer->GetDuration();
+			if ((PlayerTime > MediaDuration) && (MediaDuration > FTimespan::Zero()))
 			{
-				// It wont play if we seek to the very end, so go back a little bit.
-				PlayerTime = MediaDuration - FTimespan::FromSeconds(0.001f);
+				bool bIsPlaylist = (MediaPlaylist != nullptr) && (MediaPlaylist->Num() > 1);
+				if ((bLoop) && (bIsPlaylist == false))
+				{
+					PlayerTime %= MediaDuration;
+				}
+				else
+				{
+					// It wont play if we seek to the very end, so go back a little bit.
+					PlayerTime = MediaDuration - FTimespan::FromSeconds(0.001f);
+				}
 			}
+			TimeWhenPlaybackPaused = -1.0;
 		}
 	}
-
 	return PlayerTime;
 }
 
@@ -1003,7 +1027,7 @@ void UMediaPlateComponent::RemoveLetterboxes()
 	for (const TObjectPtr<UStaticMeshComponent>& Letterbox : Letterboxes)
 	{
 		if (Letterbox != nullptr)
-		{	
+		{
 			Letterbox->DestroyComponent();
 		}
 	}
@@ -1032,6 +1056,16 @@ void UMediaPlateComponent::OnMediaEnd()
 	Next();
 }
 
+void UMediaPlateComponent::OnMediaResumed()
+{
+	ActualPlaybackState = EPlaybackState::Playing;
+}
+
+void UMediaPlateComponent::OnMediaSuspended()
+{
+	ActualPlaybackState = EPlaybackState::Paused;
+}
+
 void UMediaPlateComponent::SetUpTextures(UE::MediaPlateComponent::ESetUpTexturesFlags Flags)
 {
 	// Prevent media texture blackouts by only updating resource and material uniforms on relevant changes.
@@ -1105,7 +1139,7 @@ void UMediaPlateComponent::SetNormalMode(bool bInIsNormalMode)
 			}
 			TextureLayers[0].Textures[0] = 0;
 			UpdateTextureLayers();
-			
+
 			ProxySetTextureBlend(0, 0, 1.0f);
 			MediaTextures[0]->SetMediaPlayer(MediaPlayer);
 		}
@@ -1147,7 +1181,7 @@ void UMediaPlateComponent::UpdateTextureLayers()
 						{
 							// Assign the next layer in the material to this layer.
 							TextureLayers[LayerIndex].MaterialLayerIndex = MaterialLayerIndex;
-							
+
 							int32 MatTexIndex = MaterialLayerIndex * MatNumTexPerLayer +
 								LayerTexIndex;
 							FString TextureName = BaseTextureName;
@@ -1251,7 +1285,7 @@ void UMediaPlateComponent::PostEditChangeProperty(FPropertyChangedEvent& Propert
 		if (MediaTextureTrackerObject != nullptr)
 		{
 			MediaTextureTrackerObject->VisibleMipsTilesCalculations = VisibleMipsTilesCalculations;
-			
+
 			RestartPlayer();
 		}
 	}
diff --git a/Engine/Plugins/Media/MediaPlate/Source/MediaPlate/Public/MediaPlateComponent.h b/Engine/Plugins/Media/MediaPlate/Source/MediaPlate/Public/MediaPlateComponent.h
--- a/Engine/Plugins/Media/MediaPlate/Source/MediaPlate/Public/MediaPlateComponent.h
+++ b/Engine/Plugins/Media/MediaPlate/Source/MediaPlate/Public/MediaPlateComponent.h
@@ -137,7 +137,7 @@ public:
 
 	/**
 	 * Call this to seek to the specified playback time.
-	 * 
+	 *
 	 * @param Time			Time to seek to.
 	 * @return				True on success, false otherwise.
 	 */
@@ -337,6 +337,19 @@ private:
 	/** Desired rate of play that we want. */
 	float CurrentRate = 0.0f;
 
+	enum class EPlaybackState
+	{
+		Unset,
+		Paused,
+		Playing,
+		Resume
+	};
+	/** State transitions. */
+	EPlaybackState IntendedPlaybackState = EPlaybackState::Unset;
+	EPlaybackState PendingPlaybackState = EPlaybackState::Unset;
+	EPlaybackState ActualPlaybackState = EPlaybackState::Unset;
+
+
 	/** If true then only allow playback when the media plate is visible. */
 	UPROPERTY(EditAnywhere, BlueprintReadOnly, Category = "Control", meta = (AllowPrivateAccess = true))
 	bool bPlayOnlyWhenVisible = false;
@@ -369,7 +382,7 @@ private:
 	UPROPERTY(EditAnywhere, Category = "EXR Tiles & Mips", meta = (DisplayName = "Adaptive Pole Mip Upscale", EditCondition = "VisibleMipsTilesCalculations == EMediaTextureVisibleMipsTiles::Sphere", EditConditionHides, AllowPrivateAccess = true))
 	bool bAdaptivePoleMipUpscaling = true;
 
-	/** If > 0, then this is the aspect ratio of our screen and 
+	/** If > 0, then this is the aspect ratio of our screen and
 	 * letterboxes will be added if the media is smaller than the screen. */
 	UPROPERTY()
 	float LetterboxAspectRatio = 0.0f;
@@ -408,7 +421,7 @@ private:
 	/** Our media clock sink. */
 	TSharedPtr<FMediaComponentClockSink, ESPMode::ThreadSafe> ClockSink;
 	/** Game time when we paused playback. */
-	float TimeWhenPlaybackPaused = 0.0f;
+	double TimeWhenPlaybackPaused = -1.0;
 	/** True if our media should be playing when visible. */
 	bool bWantsToPlayWhenVisible = false;
 	/** True if we should resume where we left off when we open the media. */
@@ -438,7 +451,7 @@ private:
 
 	/**
 	 * Plays a media source.
-	 * 
+	 *
 	 * @param	InMediaSource		Media source to play.
 	 * @param	bInPlayOnOpen		True to play, false to just open.
 	 * @return	True if we played anything.
@@ -508,6 +521,18 @@ private:
 	UFUNCTION()
 	void OnMediaEnd();
 
+	/**
+	 * Called by the media player when the video resumes.
+	 */
+	UFUNCTION()
+	void OnMediaResumed();
+
+	/**
+	 * Called by the media player when the video pauses.
+	 */
+	UFUNCTION()
+	void OnMediaSuspended();
+
 	/**
 	 * Sets up the textures we have.
 	 */
