[Media IO] Add Media Plate Actor Tools Analytics/Telemetry

Adding the following telemetry events:

- Add Media Plate Actor to Level (Editor.Usage.MediaPlate.AddMediaPlateActorToLevel)
- Add Media Plate Actor to Sequencer (Editor.Usage.MediaPlate.AddMediaPlateActorToSequencer)

diff --git a/Engine/Plugins/Media/MediaPlate/Source/MediaPlateEditor/Private/MediaPlateEditorModule.cpp b/Engine/Plugins/Media/MediaPlate/Source/MediaPlateEditor/Private/MediaPlateEditorModule.cpp
--- a/Engine/Plugins/Media/MediaPlate/Source/MediaPlateEditor/Private/MediaPlateEditorModule.cpp
+++ b/Engine/Plugins/Media/MediaPlate/Source/MediaPlateEditor/Private/MediaPlateEditorModule.cpp
@@ -3,6 +3,7 @@
 #include "MediaPlateEditorModule.h"
 
 #include "Editor.h"
+#include "EngineAnalytics.h"
 #include "IPlacementModeModule.h"
 #include "ISequencerModule.h"
 #include "LevelEditor.h"
@@ -72,6 +73,8 @@ void FMediaPlateEditorModule::ShutdownModule()
 		{
 			EditorAssetSubsystem->GetOnExtractAssetFromFile().RemoveAll(this);
 		}
+
+		GEditor->OnLevelActorAdded().RemoveAll(this);
 	}
 
 	ISequencerModule* SequencerModulePtr = FModuleManager::Get().GetModulePtr<ISequencerModule>("Sequencer");
@@ -254,6 +257,11 @@ void FMediaPlateEditorModule::OnPostEngineInit()
 	{
 		EditorAssetSubsystem->GetOnExtractAssetFromFile().AddRaw(this, &FMediaPlateEditorModule::ExtractAssetDataFromFiles);
 	}
+
+	if (GEditor)
+	{
+		GEditor->OnLevelActorAdded().AddRaw(this, &FMediaPlateEditorModule::OnLevelActorAdded);
+	}
 }
 
 void FMediaPlateEditorModule::ExtractAssetDataFromFiles(const TArray<FString>& Files,
@@ -425,6 +433,19 @@ void FMediaPlateEditorModule::UnregisterContextMenuExtender()
 	}
 }
 
+void FMediaPlateEditorModule::OnLevelActorAdded(AActor* InActor)
+{
+	if (!InActor || InActor->HasAnyFlags(RF_Transient))
+	{
+		return;
+	}
+
+	if (InActor->GetClass() == AMediaPlate::StaticClass())
+	{
+		FEngineAnalytics::GetProvider().RecordEvent(TEXT("Editor.Usage.MediaPlate.AddMediaPlateActorToLevel"));
+	}
+}
+
 IMPLEMENT_MODULE(FMediaPlateEditorModule, MediaPlateEditor)
 
 #undef LOCTEXT_NAMESPACE
diff --git a/Engine/Plugins/Media/MediaPlate/Source/MediaPlateEditor/Private/Sequencer/MediaPlateTrackEditor.cpp b/Engine/Plugins/Media/MediaPlate/Source/MediaPlateEditor/Private/Sequencer/MediaPlateTrackEditor.cpp
--- a/Engine/Plugins/Media/MediaPlate/Source/MediaPlateEditor/Private/Sequencer/MediaPlateTrackEditor.cpp
+++ b/Engine/Plugins/Media/MediaPlate/Source/MediaPlateEditor/Private/Sequencer/MediaPlateTrackEditor.cpp
@@ -3,6 +3,7 @@
 #include "Sequencer/MediaPlateTrackEditor.h"
 
 #include "ActorTreeItem.h"
+#include "EngineAnalytics.h"
 #include "Framework/MultiBox/MultiBoxBuilder.h"
 #include "ISequencer.h"
 #include "MediaTexture.h"
@@ -120,6 +121,11 @@ void FMediaPlateTrackEditor::HandleActorAdded(AActor* Actor, FGuid TargetObjectG
 {
 	if (Actor)
 	{
+		if (Actor->GetClass() == AMediaPlate::StaticClass())
+		{
+			FEngineAnalytics::GetProvider().RecordEvent(TEXT("Editor.Usage.MediaPlate.AddMediaPlateActorToSequencer"));
+		}
+
 		if (UMediaPlateComponent* MediaPlateComponent = Actor->FindComponentByClass<UMediaPlateComponent>())
 		{
 			AddTrackForComponent(MediaPlateComponent, TargetObjectGuid);
diff --git a/Engine/Plugins/Media/MediaPlate/Source/MediaPlateEditor/Public/MediaPlateEditorModule.h b/Engine/Plugins/Media/MediaPlate/Source/MediaPlateEditor/Public/MediaPlateEditorModule.h
--- a/Engine/Plugins/Media/MediaPlate/Source/MediaPlateEditor/Public/MediaPlateEditorModule.h
+++ b/Engine/Plugins/Media/MediaPlate/Source/MediaPlateEditor/Public/MediaPlateEditorModule.h
@@ -123,4 +123,9 @@ private:
 	 * Release context menu related resources such as handles.
 	 */
 	void UnregisterContextMenuExtender();
+
+	/**
+	 * Called when an Actor has been added to a Level
+	 */
+	void OnLevelActorAdded(AActor* InActor);
 };
