[Media Plate] Unable to expose video source selection to Remote Control from media plate component

This solution proposes new struct addressing source selection, called FMediaPlateResource.
This is used as a property for UMediaPlateComponent, which can be used to store information about the currently selected media source.

FMediaPlateResource can then be controlled via Remote Control, as a Property or as an Action.

The structure Deprecates the MediaPlaylist property.

diff --git a/Engine/Plugins/Media/MediaPlate/Source/MediaPlate/Private/MediaPlateComponent.cpp b/Engine/Plugins/Media/MediaPlate/Source/MediaPlate/Private/MediaPlateComponent.cpp
--- a/Engine/Plugins/Media/MediaPlate/Source/MediaPlate/Private/MediaPlateComponent.cpp
+++ b/Engine/Plugins/Media/MediaPlate/Source/MediaPlate/Private/MediaPlateComponent.cpp
@@ -4,26 +4,29 @@
 
 #include "Components/StaticMeshComponent.h"
 #include "Engine/StaticMesh.h"
+#include "IMediaAssetsModule.h"
 #include "IMediaClock.h"
 #include "IMediaClockSink.h"
 #include "IMediaModule.h"
 #include "Materials/Material.h"
 #include "Materials/MaterialInstanceDynamic.h"
-#include "MediaComponent.h"
-#include "MediaPlateModule.h"
 #include "MediaPlate.h"
+#include "MediaPlateModule.h"
 #include "MediaPlayer.h"
 #include "MediaPlaylist.h"
-#include "MediaSource.h"
 #include "MediaSoundComponent.h"
+#include "MediaSource.h"
 #include "MediaTexture.h"
 #include "MediaTextureTracker.h"
 
+#if WITH_EDITOR
+#include "ScopedTransaction.h"
+#endif
+
 #include UE_INLINE_GENERATED_CPP_BY_NAME(MediaPlateComponent)
 
 #define LOCTEXT_NAMESPACE "MediaPlate"
 
-
 namespace UE::MediaPlateComponent
 {
 	enum class ESetUpTexturesFlags
@@ -101,9 +104,6 @@ private:
 	TWeakObjectPtr<UMediaPlateComponent> Owner;
 };
 
-FLazyName UMediaPlateComponent::MediaComponentName(TEXT("MediaComponent0"));
-FLazyName UMediaPlateComponent::MediaPlaylistName(TEXT("MediaPlaylist0"));
-
 UMediaPlateComponent::UMediaPlateComponent(const FObjectInitializer& ObjectInitializer)
 	: Super(ObjectInitializer)
 {
@@ -113,8 +113,12 @@ UMediaPlateComponent::UMediaPlateComponent(const FObjectInitializer& ObjectIniti
 
 	CacheSettings.bOverride = true;
 
-	// Set up playlist.
-	MediaPlaylist = CreateDefaultSubobject<UMediaPlaylist>(MediaPlaylistName);
+#if WITH_EDITORONLY_DATA
+	PRAGMA_DISABLE_DEPRECATION_WARNINGS
+	// Still create the deprecated MediaPlaylist so the old data can be serialized during loading
+	MediaPlaylist_DEPRECATED = CreateDefaultSubobject<UMediaPlaylist>("MediaPlaylist0");
+	PRAGMA_ENABLE_DEPRECATION_WARNINGS
+#endif
 
 	// Default to plane since AMediaPlate defaults to SM_MediaPlateScreen
 	VisibleMipsTilesCalculations = EMediaTextureVisibleMipsTiles::Plane;
@@ -128,6 +132,7 @@ void UMediaPlateComponent::PostLoad()
 {
 	Super::PostLoad();
 
+PRAGMA_DISABLE_DEPRECATION_WARNINGS
 	// Use the old media texture if we have one.
 	if (MediaTexture_DEPRECATED != nullptr)
 	{
@@ -140,6 +145,11 @@ void UMediaPlateComponent::PostLoad()
 
 	UE::MediaPlateComponent::ApplyMediaTextureMipGenProperties(MediaTextureSettings, MediaTextures);
 
+	if (MediaPlaylist_DEPRECATED)
+	{
+		InitializeMediaPlateResource();
+	}
+PRAGMA_ENABLE_DEPRECATION_WARNINGS
 }
 #endif // WITH_EDITOR
 
@@ -315,9 +325,9 @@ void UMediaPlateComponent::Open()
 		if (MediaPlayer != nullptr)
 		{
 			UMediaSource* MediaSource = nullptr;
-			if (MediaPlaylist != nullptr)
+			if (UMediaPlaylist* Playlist = GetMediaPlaylist())
 			{
-				MediaSource = MediaPlaylist->Get(PlaylistIndex);
+				MediaSource = Playlist->Get(PlaylistIndex);
 			}
 			bIsPlaying = PlayMediaSource(MediaSource, bPlayOnOpen);
 		}
@@ -339,15 +349,21 @@ void UMediaPlateComponent::Open()
 
 bool UMediaPlateComponent::Next()
 {
+	UMediaPlaylist* Playlist = GetMediaPlaylist();
+	if (!Playlist)
+	{
+		return false;
+	}
+
 	bool bIsSuccessful = false;
 
 	// Do we have a playlist?
-	if ((MediaPlaylist != nullptr) && (MediaPlaylist->Num() > 1))
+	if (Playlist->Num() > 1)
 	{
-		if ((PlaylistIndex < MediaPlaylist->Num() - 1) || (bLoop))
+		if ((PlaylistIndex < Playlist->Num() - 1) || (bLoop))
 		{
 			// Get the next media to play.
-			UMediaSource* NextSource = MediaPlaylist->GetNext(PlaylistIndex);
+			UMediaSource* NextSource = Playlist->GetNext(PlaylistIndex);
 			if (NextSource != nullptr)
 			{
 				bIsSuccessful = PlayMediaSource(NextSource, true);
@@ -372,15 +388,21 @@ void UMediaPlateComponent::Pause()
 
 bool UMediaPlateComponent::Previous()
 {
+	UMediaPlaylist* Playlist = GetMediaPlaylist();
+	if (!Playlist)
+	{
+		return false;
+	}
+
 	bool bIsSuccessful = false;
 
 	// Do we have a playlist?
-	if ((MediaPlaylist != nullptr) && (MediaPlaylist->Num() > 1))
+	if (Playlist->Num() > 1)
 	{
 		// Get the previous media to play.
 		if (PlaylistIndex > 0)
 		{
-			UMediaSource* NextSource = MediaPlaylist->GetPrevious(PlaylistIndex);
+			UMediaSource* NextSource = Playlist->GetPrevious(PlaylistIndex);
 			if (NextSource != nullptr)
 			{
 				bIsSuccessful = PlayMediaSource(NextSource, true);
@@ -438,6 +460,29 @@ void UMediaPlateComponent::SetLoop(bool bInLoop)
 	}
 }
 
+UMediaPlaylist* UMediaPlateComponent::GetMediaPlaylist() const
+{
+	return MediaPlateResource.GetActivePlaylist();
+}
+
+void UMediaPlateComponent::SetMediaPlateResource(const FMediaPlateResource& InMediaPlayerResource)
+{
+#if WITH_EDITOR
+	const FScopedTransaction Transaction(LOCTEXT("SetMediaPlateResource", "Set media plate resource"));
+	this->Modify();
+	MediaPlateResource.Modify();
+#endif
+
+	MediaPlateResource.Init(InMediaPlayerResource);
+
+	if (MediaPlateResource.GetResourceType() == EMediaPlateResourceType::External)
+	{
+		MediaPlateResource.LoadExternalMedia(FString(InMediaPlayerResource.GetExternalMediaPath()), this);
+	}
+
+	RefreshMediaPlateResource();
+}
+
 void UMediaPlateComponent::SetMeshRange(FVector2D InMeshRange)
 {
 	MeshRange = InMeshRange;
@@ -525,7 +570,7 @@ bool UMediaPlateComponent::PlayMediaSource(UMediaSource* InMediaSource, bool bIn
 		// Set media options.
 		if (MediaPlayer != nullptr)
 		{
-			bool bIsPlaylist = (MediaPlaylist != nullptr) && (MediaPlaylist->Num() > 1);
+			const bool bIsPlaylist = MediaPlateResource.GetResourceType() == EMediaPlateResourceType::Playlist;
 
 			// Play the source.
 			FMediaPlayerOptions Options;
@@ -691,9 +736,10 @@ const FMediaSourceCacheSettings& UMediaPlateComponent::GetCacheSettings() const
 UMediaSource* UMediaPlateComponent::ProxyGetMediaSourceFromIndex(int32 Index) const
 {
 	UMediaSource* MediaSource = nullptr;
-	if (MediaPlaylist != nullptr)
+
+	if (UMediaPlaylist* Playlist = GetMediaPlaylist())
 	{
-		MediaSource = MediaPlaylist->Get(Index);
+		MediaSource = Playlist->Get(Index);
 	}
 	return MediaSource;
 }
@@ -816,7 +862,6 @@ void UMediaPlateComponent::ProxySetTextureBlend(int32 LayerIndex, int32 TextureI
 	}
 }
 
-#if WITH_EDITOR
 float UMediaPlateComponent::GetForwardRate(UMediaPlayer* MediaPlayer)
 {
 	float Rate = MediaPlayer->GetRate();
@@ -840,7 +885,6 @@ float UMediaPlateComponent::GetReverseRate(UMediaPlayer* MediaPlayer)
 
 	return 2.0f * Rate;
 }
-#endif
 
 void UMediaPlateComponent::RestartPlayer()
 {
@@ -920,7 +964,7 @@ FTimespan UMediaPlateComponent::GetResumeTime()
 			FTimespan MediaDuration = MediaPlayer->GetDuration();
 			if ((PlayerTime > MediaDuration) && (MediaDuration > FTimespan::Zero()))
 			{
-				bool bIsPlaylist = (MediaPlaylist != nullptr) && (MediaPlaylist->Num() > 1);
+				const bool bIsPlaylist = MediaPlateResource.GetResourceType() == EMediaPlateResourceType::Playlist;
 				if ((bLoop) && (bIsPlaylist == false))
 				{
 					PlayerTime %= MediaDuration;
@@ -1219,6 +1263,78 @@ void UMediaPlateComponent::UpdateTextureLayers()
 	}
 }
 
+#if WITH_EDITORONLY_DATA
+void UMediaPlateComponent::InitializeMediaPlateResource()
+{
+	PRAGMA_DISABLE_DEPRECATION_WARNINGS
+	UMediaPlaylist* Playlist = MediaPlaylist_DEPRECATED;
+
+	if (!Playlist)
+	{
+		return;
+	}
+
+	const int32 PlaylistElementsNum = Playlist->Num();
+
+	if (PlaylistElementsNum == 0)
+	{
+		return;
+	}
+
+	if (PlaylistElementsNum > 1)
+	{
+		// Playlist has more than one element, source type will be set to playlist
+		MediaPlateResource.SelectPlaylist(Playlist);
+	}
+	else
+	{
+		if (UMediaSource* MediaSource = Playlist->Get(0))
+		{
+			// If outer is this, source has been created from external file
+			if (MediaSource->GetOuter() == this)
+			{
+				FString MediaPath = MediaSource->GetUrl();
+
+				const FString FilePrefix(TEXT("file://"));
+				const FString ImgPrefix(TEXT("img://"));
+				if (MediaPath.StartsWith(FilePrefix))
+				{
+					MediaPath = MediaPath.RightChop(FilePrefix.Len());
+				}
+				else if (MediaPath.StartsWith(ImgPrefix))
+				{
+					MediaPath = MediaPath.RightChop(ImgPrefix.Len());
+				}
+
+				MediaPlateResource.ExternalMediaPath = MediaPath;
+				MediaPlateResource.ExternalMedia = MediaSource;
+				MediaPlateResource.SetResourceType(EMediaPlateResourceType::External);
+				MediaPlateResource.RefreshActivePlaylist(this);
+			}
+			else
+			{
+				MediaPlateResource.SelectAsset(MediaSource, this);
+			}
+		}
+	}
+
+	MediaPlaylist_DEPRECATED = nullptr;
+	PRAGMA_ENABLE_DEPRECATION_WARNINGS
+}
+#endif
+
+void UMediaPlateComponent::RefreshMediaPlateResource()
+{
+	MediaPlateResource.RefreshActivePlaylist(this);
+
+	constexpr EMediaPlateEventState State = EMediaPlateEventState::Close;
+	SwitchStates(State);
+
+	if (IMediaAssetsModule* MediaAssets = FModuleManager::LoadModulePtr<IMediaAssetsModule>("MediaAssets"))
+	{
+		MediaAssets->BroadcastOnMediaStateChangedEvent({GetOwner()->GetPathName()}, static_cast<uint8>(State));
+	}
+}
 
 #if WITH_EDITOR
 
@@ -1238,7 +1354,7 @@ void UMediaPlateComponent::PostEditChangeProperty(FPropertyChangedEvent& Propert
 {
 	Super::PostEditChangeProperty(PropertyChangedEvent);
 
-	// Has bEnableAudiio changed?
+	// Has bEnableAudio changed?
 	if (PropertyChangedEvent.GetPropertyName() == GET_MEMBER_NAME_CHECKED(ThisClass, bEnableAudio))
 	{
 		// Are we turning on audio?
@@ -1311,7 +1427,12 @@ void UMediaPlateComponent::PostEditChangeProperty(FPropertyChangedEvent& Propert
 		UE::MediaPlateComponent::ApplyMediaTextureMipGenProperties(MediaTextureSettings, MediaTextures);
 		RestartPlayer();
 	}
+	else if (PropertyChangedEvent.GetPropertyName() == GET_MEMBER_NAME_CHECKED(ThisClass, MediaPlateResource))
+	{
+		RefreshMediaPlateResource();
+	}
 }
+#endif // WITH_EDITOR
 
 void UMediaPlateComponent::SwitchStates(EMediaPlateEventState State)
 {
@@ -1367,7 +1488,5 @@ void UMediaPlateComponent::SwitchStates(EMediaPlateEventState State)
 		break;
 	}
 }
-#endif // WITH_EDITOR
 
 #undef LOCTEXT_NAMESPACE
-
diff --git a/Engine/Plugins/Media/MediaPlate/Source/MediaPlate/Private/MediaPlateResource.cpp b/Engine/Plugins/Media/MediaPlate/Source/MediaPlate/Private/MediaPlateResource.cpp
new file mode 100644
--- /dev/null
+++ b/Engine/Plugins/Media/MediaPlate/Source/MediaPlate/Private/MediaPlateResource.cpp
@@ -0,0 +1,117 @@
+﻿// Copyright Epic Games, Inc. All Rights Reserved.
+
+#include "MediaPlateResource.h"
+#include "MediaPlaylist.h"
+#include "MediaSource.h"
+
+void FMediaPlateResource::SetResourceType(EMediaPlateResourceType InType)
+{
+	Type = InType;
+
+	if (const UMediaPlaylist* Playlist = GetActivePlaylist())
+	{
+		RefreshActivePlaylist(Playlist->GetOuter());
+	}
+}
+
+void FMediaPlateResource::SelectAsset(const UMediaSource* InMediaSource, UObject* InOuter)
+{
+	SetResourceType(EMediaPlateResourceType::Asset);
+	MediaAsset = InMediaSource;
+	RefreshActivePlaylist(InOuter);
+}
+
+void FMediaPlateResource::LoadExternalMedia(const FString& InFilePath, UObject* InOuter)
+{
+	SetResourceType(EMediaPlateResourceType::External);
+	ExternalMediaPath = InFilePath;
+	ExternalMedia = UMediaSource::SpawnMediaSourceForString(InFilePath, InOuter);
+	RefreshActivePlaylist(InOuter);
+}
+
+void FMediaPlateResource::SelectPlaylist(const UMediaPlaylist* InPlaylist)
+{
+	SetResourceType(EMediaPlateResourceType::Playlist);
+	SourcePlaylist = InPlaylist;
+	ActivePlaylist = GetSourcePlaylist();
+}
+
+#if WITH_EDITOR
+void FMediaPlateResource::Modify() const
+{
+	if (ActivePlaylist)
+	{
+		ActivePlaylist->Modify();
+	}
+}
+#endif
+
+void FMediaPlateResource::Init(const FMediaPlateResource& InOther)
+{
+	if (!ExternalMediaPath.IsEmpty())
+	{
+		ExternalMediaPath = FString(InOther.GetExternalMediaPath());
+	}
+
+	if (const UMediaSource* OtherMediaAsset = InOther.GetMediaAsset())
+	{
+		MediaAsset = OtherMediaAsset;
+	}
+
+	if (const UMediaPlaylist* OtherMediaPlaylist = InOther.GetSourcePlaylist())
+	{
+		SourcePlaylist = OtherMediaPlaylist;
+	}
+
+	Type = InOther.GetResourceType();
+}
+
+void FMediaPlateResource::RefreshActivePlaylist(UObject* InOuter)
+{
+	if (Type == EMediaPlateResourceType::Playlist)
+	{
+		ActivePlaylist = GetSourcePlaylist();
+	}
+	else if (UMediaSource* MediaSource = GetSelectedMedia())
+	{
+		if (InOuter)
+		{
+			ActivePlaylist = nullptr;
+			ActivePlaylist = NewObject<UMediaPlaylist>(InOuter, NAME_None, RF_Transactional);
+			ActivePlaylist->Add(MediaSource);
+		}
+	}
+}
+
+UMediaSource* FMediaPlateResource::GetMediaAsset() const
+{
+	return MediaAsset.LoadSynchronous();
+}
+
+UMediaPlaylist* FMediaPlateResource::GetSourcePlaylist() const
+{
+	return SourcePlaylist.LoadSynchronous();
+}
+
+UMediaPlaylist* FMediaPlateResource::GetActivePlaylist() const
+{
+	return ActivePlaylist;
+}
+
+UMediaSource* FMediaPlateResource::GetSelectedMedia() const
+{
+	switch (Type)
+	{
+		case EMediaPlateResourceType::Playlist:
+			return nullptr;
+
+		case EMediaPlateResourceType::External:
+			return ExternalMedia;
+
+		case EMediaPlateResourceType::Asset:
+			return MediaAsset.LoadSynchronous();
+
+		default:
+			return nullptr;
+	}
+}
diff --git a/Engine/Plugins/Media/MediaPlate/Source/MediaPlate/Public/MediaPlateComponent.h b/Engine/Plugins/Media/MediaPlate/Source/MediaPlate/Public/MediaPlateComponent.h
--- a/Engine/Plugins/Media/MediaPlate/Source/MediaPlate/Public/MediaPlateComponent.h
+++ b/Engine/Plugins/Media/MediaPlate/Source/MediaPlate/Public/MediaPlateComponent.h
@@ -4,6 +4,7 @@
 
 #include "Components/ActorComponent.h"
 #include "MediaPlayerProxyInterface.h"
+#include "MediaPlateResource.h"
 #include "MediaSource.h"
 #include "MediaTextureTracker.h"
 #include "Misc/EnumClassFlags.h"
@@ -74,7 +75,6 @@ public:
 	virtual void BeginDestroy() override;
 	virtual void TickComponent(float DeltaTime, enum ELevelTick TickType, FActorComponentTickFunction* ThisTickFunction) override;
 	virtual void OnUnregister() override;
-
 #if WITH_EDITOR
 	virtual void PostEditChangeProperty(FPropertyChangedEvent& PropertyChangedEvent) override;
 #endif // WITH_EDITOR
@@ -168,6 +168,16 @@ public:
 	UFUNCTION(BlueprintSetter)
 	void SetLoop(bool bInLoop);
 
+	/**
+	 * Get the currently active Media Playlist
+	 */
+	UMediaPlaylist* GetMediaPlaylist() const;
+
+	/**
+	 * Update Media Player Resource. This will also refresh Playlist accordingly.
+	 */
+	void SetMediaPlateResource(const FMediaPlateResource& InMediaPlayerResource);
+
 	/** If set then play when opening the media. */
 	UPROPERTY(EditAnywhere, BlueprintReadWrite, Category = "Control")
 	bool bPlayOnOpen = true;
@@ -196,9 +206,15 @@ public:
 	UPROPERTY(EditAnywhere, Category = "Advanced|Other")
 	TArray<TObjectPtr<UStaticMeshComponent>> Letterboxes;
 
-	/** What media playlist to play. */
-	UPROPERTY(BlueprintReadWrite, Category = "MediaPlate")
-	TObjectPtr<UMediaPlaylist> MediaPlaylist;
+#if WITH_EDITORONLY_DATA
+	UE_DEPRECATED(5.5, "Use MediaPlateResource instead")
+	UPROPERTY(meta=(DeprecatedProperty, DeprecationMessage="Use MediaPlateResource instead"))
+	TObjectPtr<UMediaPlaylist> MediaPlaylist_DEPRECATED;
+#endif
+
+	/** Which media source is used to populate the media playlist */
+	UPROPERTY(EditAnywhere, BlueprintReadWrite, Category="MediaPlate", meta=(AllowPrivateAccess="true"), Setter=SetMediaPlateResource)
+	FMediaPlateResource MediaPlateResource;
 
 	/** The current index of the source in the play list being played. */
 	UPROPERTY(BlueprintReadWrite, Category = "MediaPlate")
@@ -273,12 +289,12 @@ public:
 	 * Call this to set the mip tile calculations mesh mode. (Note: restarts playback to apply changes.)
 	 */
 	void SetVisibleMipsTilesCalculations(EMediaTextureVisibleMipsTiles InVisibleMipsTilesCalculations);
+#endif
 
 	/**
-	* Called whenever a button was pressed locally or on a remote endpoint.
-	*/
+	 * Called whenever a button was pressed locally or on a remote endpoint.
+	 */
 	void SwitchStates(EMediaPlateEventState State);
-#endif
 
 	/**
 	 * Called from the media clock.
@@ -296,7 +312,6 @@ public:
 	virtual bool ProxySetAspectRatio(UMediaPlayer* InMediaPlayer) override;
 	virtual void ProxySetTextureBlend(int32 LayerIndex, int32 TextureIndex, float Blend) override;
 
-#if WITH_EDITOR
 public:
 	/**
 	 * Get the rate to use when we press the forward button.
@@ -307,7 +322,6 @@ public:
 	 * Get the rate to use when we press the reverse button.
 	 */
 	static float GetReverseRate(UMediaPlayer* MediaPlayer);
-#endif
 
 private:
 	/**
@@ -547,4 +561,20 @@ private:
 	 * Sets textures in our material according to the layer assignments.
 	 */
 	void UpdateTextureLayers();
+
+#if WITH_EDITORONLY_DATA
+	/**
+	 * If Media Player Resource is not initialized yet (e.g. older assets)
+	 * this function will initialize it with the proper data.
+	 * Logic is the following: if media playlist has more than one entry, then the type will be set to Playlist.
+	 * If playlist has only one entry, then the function will check for that entry Outer.
+	 * If the Outer is this media plate component, then type will be set to External, otherwise to Asset.
+	 */
+	void InitializeMediaPlateResource();
+#endif
+
+	/**
+	 * Can be called after updating MediaPlayerResource, to ensure the Playlist and player state are up to date
+	 */
+	void RefreshMediaPlateResource();
 };
diff --git a/Engine/Plugins/Media/MediaPlate/Source/MediaPlate/Public/MediaPlateResource.h b/Engine/Plugins/Media/MediaPlate/Source/MediaPlate/Public/MediaPlateResource.h
new file mode 100644
--- /dev/null
+++ b/Engine/Plugins/Media/MediaPlate/Source/MediaPlate/Public/MediaPlateResource.h
@@ -0,0 +1,127 @@
+﻿// Copyright Epic Games, Inc. All Rights Reserved.
+
+#pragma once
+
+#include "UObject/ObjectPtr.h"
+#include "UObject/SoftObjectPtr.h"
+#include "MediaPlateResource.generated.h"
+
+class UMediaPlaylist;
+class UMediaSource;
+
+UENUM()
+enum class EMediaPlateResourceType : uint8
+{
+	Playlist,
+	External,
+	Asset
+};
+
+/**
+ * Helper struct to wrap source selection functionality,
+ * and enabling the usage of media source properties for places like Remote Control.
+ *
+ * This struct allows to choose between Asset, External File, Playlist options.
+ * It's mainly conceived to be used by MediaPlateComponent.
+ *
+ * See FMediaPlayerResourceCustomization class for its customization.
+ */
+USTRUCT(BlueprintType)
+struct FMediaPlateResource
+{
+	GENERATED_BODY()
+
+public:
+	/**
+	 * Returns the currently selected Media Source, based on current source Type.
+	 * If current type is Playlist, it will return nullptr
+	 */
+	MEDIAPLATE_API UMediaSource* GetSelectedMedia() const;
+
+	/**
+	 * Returns the current source playlist, if any
+	 */
+	MEDIAPLATE_API UMediaPlaylist* GetSourcePlaylist() const;
+
+	/**
+	 * Returns the active playlist, if any
+	 */
+	MEDIAPLATE_API UMediaPlaylist* GetActivePlaylist() const;
+
+	/**
+	 * Returns the current external media path, if any
+	 */
+	MEDIAPLATE_API FStringView GetExternalMediaPath() const { return ExternalMediaPath; }
+
+	/**
+	 * Returns the current asset-based Media Source, if any
+	 */
+	MEDIAPLATE_API UMediaSource* GetMediaAsset() const;
+
+	/**
+	 * Return current resource type
+	 */
+	MEDIAPLATE_API EMediaPlateResourceType GetResourceType() const { return Type; }
+
+	/**
+	 * Set current resource type
+	 */
+	MEDIAPLATE_API void SetResourceType(EMediaPlateResourceType InType);
+
+	/**
+	 * Select asset based media source. Will also update source type to Asset
+	 */
+	MEDIAPLATE_API void SelectAsset(const UMediaSource* InMediaSource, UObject* InOuter);
+
+	/**
+	 * Loads the external source at the specified path, creating a media source with the specified outer.
+	 * Will also update source type to External
+	 */
+	MEDIAPLATE_API void LoadExternalMedia(const FString& InFilePath, UObject* InOuter);
+
+	/**
+	 * Select the specified playlist. Will also update source type to Playlist
+	 */
+	MEDIAPLATE_API void SelectPlaylist(const UMediaPlaylist* InPlaylist);
+
+#if WITH_EDITOR
+	/**
+	 * Note that the Active Playlist will be modified.
+	 * Convenience function to avoid having to check and get the Playlist everytime to call Modify() on it.
+	 */
+	MEDIAPLATE_API void Modify() const;
+#endif
+
+private:
+friend class UMediaPlateComponent;
+	/**
+	 * Initialize member properties from another MediaPlayerResource. Empty or null property will not be copied over.
+	 */
+	void Init(const FMediaPlateResource& InOther);
+
+	void RefreshActivePlaylist(UObject* InOuter);
+
+	/** Media Source Type */
+	UPROPERTY()
+	EMediaPlateResourceType Type = EMediaPlateResourceType::Asset;
+
+	/** A path pointing to an external media resource */
+	UPROPERTY()
+	FString ExternalMediaPath;
+
+	/** Media Source loaded from external path */
+	UPROPERTY()
+	TObjectPtr<UMediaSource> ExternalMedia;
+
+	/** Media source coming from MediaSource asset*/
+	UPROPERTY()
+	TSoftObjectPtr<UMediaSource> MediaAsset;
+
+	/** User facing Playlist asset */
+	UPROPERTY()
+	TSoftObjectPtr<UMediaPlaylist> SourcePlaylist;
+
+	/** Currently running Playlist asset */
+	UPROPERTY()
+	TObjectPtr<UMediaPlaylist> ActivePlaylist;
+};
diff --git a/Engine/Plugins/Media/MediaPlate/Source/MediaPlateEditor/Private/ActorFactoryMediaPlate.cpp b/Engine/Plugins/Media/MediaPlate/Source/MediaPlateEditor/Private/ActorFactoryMediaPlate.cpp
--- a/Engine/Plugins/Media/MediaPlate/Source/MediaPlateEditor/Private/ActorFactoryMediaPlate.cpp
+++ b/Engine/Plugins/Media/MediaPlate/Source/MediaPlateEditor/Private/ActorFactoryMediaPlate.cpp
@@ -57,42 +57,35 @@ void UActorFactoryMediaPlate::SetUpActor(UObject* Asset, AActor* Actor)
 		UMediaSource* MediaSource = Cast<UMediaSource>(Asset);
 		if ((MediaSource != nullptr) && (MediaPlate->MediaPlateComponent != nullptr))
 		{
-			UMediaPlaylist* Playlist = MediaPlate->MediaPlateComponent->MediaPlaylist;
-			if (Playlist != nullptr)
+			// Is this media source from a drag and drop?
+			FMediaPlateEditorModule* EditorModule = FModuleManager::LoadModulePtr<FMediaPlateEditorModule>("MediaPlateEditor");
+			if (EditorModule != nullptr)
 			{
-				// This gets called twice so only add the media source once.
-				if (Playlist->Num() == 0)
+				bool bIsInDragDropCache = EditorModule->RemoveMediaSourceFromDragDropCache(MediaSource);
+				if (bIsInDragDropCache && MediaSource->GetOuter() == GetTransientPackage())
 				{
-					Playlist->Add(MediaSource);
-
-					// Is this media source from a drag and drop?
-					FMediaPlateEditorModule* EditorModule = FModuleManager::LoadModulePtr<FMediaPlateEditorModule>("MediaPlateEditor");
-					if (EditorModule != nullptr)
+					// Yes. Move this out of transient.
+					// Can't do it here as the asset is still being used.
+					TWeakObjectPtr<UMediaPlateComponent> MediaPlateComponentPtr(MediaPlate->MediaPlateComponent);
+					AsyncTask(ENamedThreads::GameThread, [MediaPlateComponentPtr]()
 					{
-						bool bIsInDragDropCache = EditorModule->RemoveMediaSourceFromDragDropCache(MediaSource);
-						if (bIsInDragDropCache)
+						UMediaPlateComponent* MediaPlateComponent = MediaPlateComponentPtr.Get();
+						if (MediaPlateComponent != nullptr)
 						{
-							// Yes. Move this out of transient.
-							// Can't do it here as the asset is still being used.
-							TWeakObjectPtr<UMediaPlateComponent> MediaPlateComponentPtr(MediaPlate->MediaPlateComponent);
-							AsyncTask(ENamedThreads::GameThread, [MediaPlateComponentPtr]()
+							if (UMediaSource* MediaSource = MediaPlateComponent->MediaPlateResource.GetMediaAsset())
 							{
-								UMediaPlateComponent* MediaPlateComponent = MediaPlateComponentPtr.Get();
-								if (MediaPlateComponent != nullptr)
-								{
-									UMediaPlaylist* Playlist = MediaPlateComponent->MediaPlaylist;
-									if ((Playlist != nullptr) && (Playlist->Num() > 0))
-									{
-										UMediaSource* MediaSource = Playlist->Get(0);
-										if (MediaSource != nullptr)
-										{
-											MediaSource->Rename(nullptr, MediaPlateComponent);
-										}
-									}
-								}
-							});
+								MediaSource->Rename(nullptr, MediaPlateComponent);
+
+								// Let's initialize the plate source with the now non-transient MediaSource
+								MediaPlateComponent->MediaPlateResource.SelectAsset(MediaSource, MediaPlateComponent);
+							}
 						}
-					}
+					});
+				}
+				else
+				{
+					// MediaSource is non-transient, we can initialize the plate source right away
+					MediaPlate->MediaPlateComponent->MediaPlateResource.SelectAsset(MediaSource, MediaPlate->MediaPlateComponent);
 				}
 			}
 		}
@@ -100,4 +93,3 @@ void UActorFactoryMediaPlate::SetUpActor(UObject* Asset, AActor* Actor)
 }
 
 #undef LOCTEXT_NAMESPACE
-
diff --git a/Engine/Plugins/Media/MediaPlate/Source/MediaPlateEditor/Private/MediaPlateCustomization.cpp b/Engine/Plugins/Media/MediaPlate/Source/MediaPlateEditor/Private/MediaPlateCustomization.cpp
--- a/Engine/Plugins/Media/MediaPlate/Source/MediaPlateEditor/Private/MediaPlateCustomization.cpp
+++ b/Engine/Plugins/Media/MediaPlate/Source/MediaPlateEditor/Private/MediaPlateCustomization.cpp
@@ -16,20 +16,17 @@
 #include "MediaPlateEditorStyle.h"
 #include "MediaPlayer.h"
 #include "MediaPlayerEditorModule.h"
-#include "MediaPlaylist.h"
-#include "MediaSource.h"
 #include "PropertyCustomizationHelpers.h"
 #include "ScopedTransaction.h"
 #include "Styling/AppStyle.h"
 #include "Subsystems/AssetEditorSubsystem.h"
 #include "ThumbnailRendering/ThumbnailManager.h"
+#include "Widgets/Images/SImage.h"
 #include "Widgets/Input/SButton.h"
 #include "Widgets/Input/SCheckBox.h"
 #include "Widgets/Input/SComboButton.h"
-#include "Widgets/Input/SFilePathPicker.h"
 #include "Widgets/Input/SNumericEntryBox.h"
 #include "Widgets/Input/SSegmentedControl.h"
-#include "Widgets/Images/SImage.h"
 #include "Widgets/SMediaPlateEditorMediaDetails.h"
 
 #define LOCTEXT_NAMESPACE "FMediaPlateCustomization"
@@ -39,7 +36,6 @@
 
 
 FMediaPlateCustomization::FMediaPlateCustomization()
-	: bIsMediaSourceAsset(false)
 {
 }
 
@@ -82,105 +78,15 @@ void FMediaPlateCustomization::CustomizeDetails(IDetailLayoutBuilder& DetailBuil
 		}
 	}
 
-	// Set up media source.
-	UpdateIsMediaSourceAsset();
-	UpdateMediaPath();
-
 	// Add mesh customization.
 	AddMeshCustomization(GeometryCategory);
 
-	// Add playlist.
-	PlaylistCategory.AddCustomRow(FText::FromString("Playlist Asset"))
-		.NameContent()
-		[
-			SNew(STextBlock)
-				.Text(LOCTEXT("PlaylistAsset", "Playlist Asset"))
-				.ToolTipText(LOCTEXT("PlaylistAsset_Tooltip",
-					"The playlist asset to use. If none then an internal asset will be used."))
-				.Font(IDetailLayoutBuilder::GetDetailFont())
-		]
-		.ValueContent()
-		[
-			SNew(SObjectPropertyEntryBox)
-				.AllowedClass(UMediaPlaylist::StaticClass())
-				.ObjectPath(this, &FMediaPlateCustomization::GetPlaylistPath)
-				.OnObjectChanged(this, &FMediaPlateCustomization::OnPlaylistChanged)
-		];
-
-	// Add media source type.
-	TAttribute<EVisibility> MediaSourceAssetVisibility(this, &FMediaPlateCustomization::ShouldShowMediaSourceAsset);
-	TAttribute<EVisibility> MediaSourceFileVisibility(this, &FMediaPlateCustomization::ShouldShowMediaSourceFile);
-	PlaylistCategory.AddCustomRow(FText::FromString("First Item In Playlist"))
-		.NameContent()
-		[
-			SNew(STextBlock)
-				.Text(LOCTEXT("FirstItemInPlaylist", "First Item In Playlist"))
-				.ToolTipText(LOCTEXT("FirstItemInPlaylist_ToolTip",
-					"The type of the first item in the playlist.\n"
-					"File lets you select a media file like a .mp4.\n"
-					"Asset lets you select a Media Source asset."))
-				.Font(IDetailLayoutBuilder::GetDetailFont())
-		]
-		.ValueContent()
-		[
-			SNew(SSegmentedControl<bool>)
-				.Value_Lambda([this]()
-				{
-					return bIsMediaSourceAsset;
-				})
-				.OnValueChanged(this, &FMediaPlateCustomization::SetIsMediaSourceAsset)
-
-			+ SSegmentedControl<bool>::Slot(false)
-				.Text(LOCTEXT("File", "File"))
-				.ToolTip(LOCTEXT("File_ToolTip",
-					"Select this if you want to use a file like a .mp4."))
-
-			+ SSegmentedControl<bool>::Slot(true)
-				.Text(LOCTEXT("Asset", "Asset"))
-				.ToolTip(LOCTEXT("Asset_ToolTip",
-					"Select this if you want to use a Media Source asset."))
-		];
-
-	// Add media asset.
-	PlaylistCategory.AddCustomRow(FText::FromString("Media Source Asset"))
-		.Visibility(MediaSourceAssetVisibility)
-		.NameContent()
-		[
-			SNew(STextBlock)
-				.Text(LOCTEXT("MediaAsset", "Media Source Asset"))
-				.ToolTipText(LOCTEXT("MediaAsset_ToolTip", "The Media Source to play."))
-				.Font(IDetailLayoutBuilder::GetDetailFont())
-		]
-		.ValueContent()
-		[
-			SNew(SObjectPropertyEntryBox)
-				.AllowedClass(UMediaSource::StaticClass())
-				.ObjectPath(this, &FMediaPlateCustomization::GetMediaSourcePath)
-				.OnObjectChanged(this, &FMediaPlateCustomization::OnMediaSourceChanged)
-		];
-
-	// Add media path.
-	FString FileTypeFilter = TEXT("All files (*.*)|*.*");
-	PlaylistCategory.AddCustomRow(FText::FromString("Media Path"))
-		.Visibility(MediaSourceFileVisibility)
-		.NameContent()
-		[
-			SNew(STextBlock)
-				.Text(LOCTEXT("MediaPath", "Media Path"))
-				.ToolTipText(LOCTEXT("MediaPath_ToolTip", "The path to a media file like a .mp4."))
-				.Font(IDetailLayoutBuilder::GetDetailFont())
-		]
-		.ValueContent()
-		[
-			SNew(SFilePathPicker)
-				.BrowseButtonImage(FAppStyle::GetBrush("PropertyWindow.Button_Ellipsis"))
-				.BrowseButtonStyle(FAppStyle::Get(), "HoverHintOnly")
-				.BrowseButtonToolTip(LOCTEXT("FileButtonToolTipText", "Choose a file from this computer"))
-				.BrowseTitle(LOCTEXT("PropertyEditorTitle", "File picker..."))
-				.FilePath(this, &FMediaPlateCustomization::HandleMediaPath)
-				.FileTypeFilter(FileTypeFilter)
-				.OnPathPicked(this, &FMediaPlateCustomization::HandleMediaPathPicked)
-		];
+	// Add media plate source 
+	MediaPlateResourcePropertyHandle = DetailBuilder.GetProperty(GET_MEMBER_NAME_CHECKED(UMediaPlateComponent, MediaPlateResource));
+	if (MediaPlateResourcePropertyHandle)
+	{
+		PlaylistCategory.AddProperty(MediaPlateResourcePropertyHandle);
+	}
 
 	// Add media player playback slider
 	if (IMediaPlayerEditorModule* MediaPlayerEditorModule = FModuleManager::LoadModulePtr<IMediaPlayerEditorModule>("MediaPlayerEditor"))
@@ -760,7 +666,7 @@ TSharedRef<SWidget> FMediaPlateCustomization::OnGetLetterboxAspectRatios()
 {
 	FMenuBuilder MenuBuilder(true, NULL);
 	AddAspectRatiosToMenuBuilder(MenuBuilder, &FMediaPlateCustomization::SetLetterboxAspectRatio);
-	
+
 	FUIAction Action = FUIAction(FExecuteAction::CreateSP(this, &FMediaPlateCustomization::SetLetterboxAspectRatio, 0.0f));
 	MenuBuilder.AddMenuEntry(LOCTEXT("Disable", "Disable"), FText(), FSlateIcon(), Action);
 
@@ -961,191 +867,6 @@ void FMediaPlateCustomization::OnStaticMeshChanged(const FAssetData& AssetData)
 	}
 }
 
-void FMediaPlateCustomization::UpdateIsMediaSourceAsset()
-{
-	bIsMediaSourceAsset = false;
-
-	// Get the first media plate.
-	if (MediaPlatesList.Num() > 0)
-	{
-		UMediaPlateComponent* MediaPlate = MediaPlatesList[0].Get();
-		if (MediaPlate != nullptr)
-		{
-			UMediaPlaylist* Playlist = MediaPlate->MediaPlaylist;
-			if (Playlist != nullptr)
-			{
-				// Get the first media source in the playlist.
-				UMediaSource* MediaSource = Playlist->Get(0);
-				if (MediaSource != nullptr)
-				{
-					bIsMediaSourceAsset = MediaSource->GetOuter() != MediaPlate;
-				}
-			}
-		}
-	}
-}
-
-void FMediaPlateCustomization::SetIsMediaSourceAsset(bool bIsAsset)
-{
-	if (bIsMediaSourceAsset != bIsAsset)
-	{
-		bIsMediaSourceAsset = bIsAsset;
-
-		{
-			const FScopedTransaction Transaction(LOCTEXT("SetIsMediaSourceAsset", "Media Source Playlist Changed"));
-
-			// Clear out asset.
-			for (TWeakObjectPtr<UMediaPlateComponent>& MediaPlatePtr : MediaPlatesList)
-			{
-				UMediaPlateComponent* MediaPlate = MediaPlatePtr.Get();
-				if (MediaPlate != nullptr)
-				{
-					// Get playlist.
-					UMediaPlaylist* Playlist = MediaPlate->MediaPlaylist;
-					if (Playlist != nullptr)
-					{
-						// Update playlist.
-						if (Playlist->Num() > 0)
-						{
-							Playlist->Modify();
-							Playlist->Replace(0, nullptr);
-						}
-					}
-					break;
-				}
-			}
-		}
-
-		StopMediaPlates();
-		UpdateMediaPath();
-	}
-}
-
-EVisibility FMediaPlateCustomization::ShouldShowMediaSourceAsset() const
-{
-	return bIsMediaSourceAsset ? EVisibility::Visible : EVisibility::Hidden;
-}
-
-EVisibility FMediaPlateCustomization::ShouldShowMediaSourceFile() const
-{
-	return bIsMediaSourceAsset ? EVisibility::Hidden : EVisibility::Visible;
-}
-
-FString FMediaPlateCustomization::GetMediaSourcePath() const
-{
-	FString Path;
-
-	// Get the first media plate.
-	if (MediaPlatesList.Num() > 0)
-	{
-		UMediaPlateComponent* MediaPlate = MediaPlatesList[0].Get();
-		if (MediaPlate != nullptr)
-		{
-			UMediaPlaylist* Playlist = MediaPlate->MediaPlaylist;
-			if (Playlist != nullptr)
-			{
-				// Get the first media source in the playlist.
-				UMediaSource* MediaSource = Playlist->Get(0);
-				if (MediaSource != nullptr)
-				{
-					Path = MediaSource->GetPathName();
-				}
-			}
-		}
-	}
-
-	return Path;
-}
-
-FString FMediaPlateCustomization::GetPlaylistPath() const
-{
-	FString Path;
-
-	// Get the first media plate.
-	if (MediaPlatesList.Num() > 0)
-	{
-		UMediaPlateComponent* MediaPlate = MediaPlatesList[0].Get();
-		if (MediaPlate != nullptr)
-		{
-			UMediaPlaylist* Playlist = MediaPlate->MediaPlaylist;
-			if (Playlist != nullptr)
-			{
-				if (Playlist->GetOuter() != MediaPlate)
-				{
-					Path = Playlist->GetPathName();
-				}
-			}
-		}
-	}
-
-	return Path;
-}
-
-void FMediaPlateCustomization::OnPlaylistChanged(const FAssetData& AssetData)
-{
-	StopMediaPlates();
-
-	{
-		const FScopedTransaction Transaction(LOCTEXT("OnPlaylistChanged", "Media Playlist Changed"));
-		
-		// Update the playlist.
-		UMediaPlaylist* Playlist = Cast<UMediaPlaylist>(AssetData.GetAsset());
-		for (TWeakObjectPtr<UMediaPlateComponent>& MediaPlatePtr : MediaPlatesList)
-		{
-			UMediaPlateComponent* MediaPlate = MediaPlatePtr.Get();
-			if (MediaPlate != nullptr)
-			{
-				MediaPlate->Modify();
-				MediaPlate->MediaPlaylist = Playlist;
-			}
-		}
-	}
-
-	UpdateIsMediaSourceAsset();
-	UpdateMediaPath();
-}
-
-void FMediaPlateCustomization::OnMediaSourceChanged(const FAssetData& AssetData)
-{
-	// Update the playlist with the new media source.
-	UMediaSource* MediaSource = Cast<UMediaSource>(AssetData.GetAsset());
-
-	{
-		const FScopedTransaction Transaction(LOCTEXT("OnMediaSourceChanged", "Media Source Playlist Changed"));
-
-		for (TWeakObjectPtr<UMediaPlateComponent>& MediaPlatePtr : MediaPlatesList)
-		{
-			UMediaPlateComponent* MediaPlate = MediaPlatePtr.Get();
-			if (MediaPlate != nullptr)
-			{
-				// Get playlist.
-				TObjectPtr<UMediaPlaylist>& Playlist = MediaPlate->MediaPlaylist;
-				if (Playlist == nullptr)
-				{
-					Playlist = NewObject<UMediaPlaylist>(MediaPlate, NAME_None, RF_Transactional);
-				}
-				else
-				{
-					Playlist->Modify();
-				}
-
-				// Update playlist.
-				if (Playlist->Num() > 0)
-				{
-					Playlist->Replace(0, MediaSource);
-				}
-				else
-				{
-					Playlist->Add(MediaSource);
-				}
-			}
-		}
-	}
-
-	StopMediaPlates();
-	UpdateMediaPath();
-}
-
 void FMediaPlateCustomization::OnButtonEvent(EMediaPlateEventState State)
 {
 	IMediaAssetsModule* MediaAssets = FModuleManager::LoadModulePtr<IMediaAssetsModule>("MediaAssets");
@@ -1159,7 +880,7 @@ void FMediaPlateCustomization::OnButtonEvent(EMediaPlateEventState State)
 		}
 
 		ActorsPathNames.Add(MediaPlatePtr->GetOwner()->GetPathName());
-			
+
 		if (State == EMediaPlateEventState::Open)
 		{
 			// Tell the editor module that this media plate is playing.
@@ -1175,97 +896,6 @@ void FMediaPlateCustomization::OnButtonEvent(EMediaPlateEventState State)
 	MediaAssets->BroadcastOnMediaStateChangedEvent(ActorsPathNames, (uint8)State);
 }
 
-void FMediaPlateCustomization::UpdateMediaPath()
-{
-	MediaPath.Empty();
-
-	// Get the first media plate.
-	if (MediaPlatesList.Num() > 0)
-	{
-		UMediaPlateComponent* MediaPlate = MediaPlatesList[0].Get();
-		if (MediaPlate != nullptr)
-		{
-			UMediaPlaylist* Playlist = MediaPlate->MediaPlaylist;
-			if (Playlist != nullptr)
-			{
-				// Get the first media source in the playlist.
-				UMediaSource* MediaSource = Playlist->Get(0);
-				if (MediaSource != nullptr)
-				{
-					MediaPath = MediaSource->GetUrl();
-
-					// Remove certain types.
-					const FString FilePrefix(TEXT("file://"));
-					const FString ImgPrefix(TEXT("img://"));
-					if (MediaPath.StartsWith(FilePrefix))
-					{
-						MediaPath = MediaPath.RightChop(FilePrefix.Len());
-					}
-					else if (MediaPath.StartsWith(ImgPrefix))
-					{
-						MediaPath = MediaPath.RightChop(ImgPrefix.Len());
-					}
-				}
-			}
-		}
-	}
-}
-
-FString FMediaPlateCustomization::HandleMediaPath() const
-{
-	return MediaPath;
-}
-
-void FMediaPlateCustomization::HandleMediaPathPicked(const FString& PickedPath)
-{
-	// Did we get something?
-	if ((PickedPath.IsEmpty() == false) && (PickedPath != MediaPath))
-	{
-		// Stop playback.
-		StopMediaPlates();
-
-		{
-			const FScopedTransaction Transaction(LOCTEXT("OnMediaPathChanged", "Media Path Changed"));
-
-			// Set up media source for our media plates.
-			for (TWeakObjectPtr<UMediaPlateComponent>& MediaPlatePtr : MediaPlatesList)
-			{
-				UMediaPlateComponent* MediaPlate = MediaPlatePtr.Get();
-				if (MediaPlate != nullptr)
-				{
-					// Get playlist.
-					TObjectPtr<UMediaPlaylist>& Playlist = MediaPlate->MediaPlaylist;
-					if (Playlist == nullptr)
-					{
-						Playlist = NewObject<UMediaPlaylist>(MediaPlate, NAME_None, RF_Transactional);
-					}
-					else
-					{
-						Playlist->Modify();
-					}
-
-					// Create media source for this path.
-					UMediaSource* MediaSource = UMediaSource::SpawnMediaSourceForString(PickedPath, MediaPlate);
-					if (MediaSource != nullptr)
-					{
-						if (Playlist->Num() > 0)
-						{
-							Playlist->Replace(0, MediaSource);
-						}
-						else
-						{
-							Playlist->Add(MediaSource);
-						}
-					}
-				}
-			}
-		}
-
-		// Update the media path.
-		UpdateMediaPath();
-	}
-}
-
 FReply FMediaPlateCustomization::OnOpenMediaPlate()
 {
 	// Get all our objects.
diff --git a/Engine/Plugins/Media/MediaPlate/Source/MediaPlateEditor/Private/MediaPlateCustomization.h b/Engine/Plugins/Media/MediaPlate/Source/MediaPlateEditor/Private/MediaPlateCustomization.h
--- a/Engine/Plugins/Media/MediaPlate/Source/MediaPlateEditor/Private/MediaPlateCustomization.h
+++ b/Engine/Plugins/Media/MediaPlate/Source/MediaPlateEditor/Private/MediaPlateCustomization.h
@@ -9,9 +9,10 @@
 #include "MediaTextureTracker.h"
 #include "Styling/SlateTypes.h"
 
+class FMenuBuilder;
 class IDetailCategoryBuilder;
 class IDetailLayoutBuilder;
-class FMenuBuilder;
+class IPropertyHandle;
 class SWidget;
 class UMediaPlateComponent;
 class UMediaPlayer;
@@ -48,8 +49,6 @@ private:
 	FDelegateHandle PropertyChangeDelegate;
 	/** List of the media plates we are editing. */
 	TArray<TWeakObjectPtr<UMediaPlateComponent>> MediaPlatesList;
-	/** Stores the current value of the MediaPath property. */
-	FString MediaPath;
 
 	/** Whether we have a plane, sphere, etc. */
 	EMediaTextureVisibleMipsTiles MeshMode;
@@ -57,9 +56,9 @@ private:
 	/** Handles mesh stuff. */
 	FMediaPlateCustomizationMesh MeshCustomization;
 
-	/** True if the first media source in the playlist is an external asset. */
-	bool bIsMediaSourceAsset;
-	
+	/** Property handle of the currently customized Media Plate Resource */
+	TSharedPtr<IPropertyHandle> MediaPlateResourcePropertyHandle;
+
 	/**
 	 * Adds widgets for editing the mesh.
 	 */
@@ -171,66 +170,11 @@ private:
 	 */
 	void OnStaticMeshChanged(const FAssetData& AssetData);
 
-	/**
-	 * Updates bIsMediaSourceAsset depending on the current playlist.
-	 */
-	void UpdateIsMediaSourceAsset();
-
-	/**
-	 * Call this to enable/disable automatic aspect ratio.
-	 */
-	void SetIsMediaSourceAsset(bool bIsAsset);
-
-	/**
-	 * Controls visibility for widgets when the media source is an asset.
-	 */
-	EVisibility ShouldShowMediaSourceAsset() const;
-
-	/**
-	 * Controls visibility for widgets when the media source is a file.
-	 */
-	EVisibility ShouldShowMediaSourceFile() const;
-
-	/**
-	 * Gets the object path for the media source object.
-	 */
-	FString GetMediaSourcePath() const;
-
-	/**
-	 * Gets the object path for the media playlist object.
-	 */
-	FString GetPlaylistPath() const;
-
-	/**
-	 * Called when the playlist changes.
-	 */
-	void OnPlaylistChanged(const FAssetData& AssetData);
-
-	/**
-	 * Called when the media source widget changes.
-	 */
-	void OnMediaSourceChanged(const FAssetData& AssetData);
-
 	/**
 	* Changes the state of selected media plates and broadcasts the even to the remote endpoints.
 	*/
 	void OnButtonEvent(EMediaPlateEventState State);
 
-	/**
-	 * Updates MediaPath from the current MediaSource.
-	 */
-	void UpdateMediaPath();
-
-	/**
-	 * Called to get the media path for the file picker.
-	 */
-	FString HandleMediaPath() const;
-	
-	/**
-	 * Called when we select a media path.
-	 */
-	void HandleMediaPathPicked(const FString& PickedPath);
-
 	/**
 	 * Called when the open media plate button is pressed.
 	 */
@@ -243,4 +187,3 @@ private:
 
 	UMediaPlayer* GetMediaPlayer() const;
 };
-
diff --git a/Engine/Plugins/Media/MediaPlate/Source/MediaPlateEditor/Private/MediaPlateEditorModule.cpp b/Engine/Plugins/Media/MediaPlate/Source/MediaPlateEditor/Private/MediaPlateEditorModule.cpp
--- a/Engine/Plugins/Media/MediaPlate/Source/MediaPlateEditor/Private/MediaPlateEditorModule.cpp
+++ b/Engine/Plugins/Media/MediaPlate/Source/MediaPlateEditor/Private/MediaPlateEditorModule.cpp
@@ -13,7 +13,7 @@
 #include "MediaPlateComponent.h"
 #include "MediaPlateCustomization.h"
 #include "MediaPlateEditorStyle.h"
-#include "MediaPlayer.h"
+#include "MediaPlateResourceCustomization.h"
 #include "MediaSoundComponent.h"
 #include "MediaSource.h"
 #include "Models/MediaPlateEditorCommands.h"
@@ -40,11 +40,15 @@ void FMediaPlateEditorModule::StartupModule()
 
 	// Register customizations.
 	MediaPlateName = UMediaPlateComponent::StaticClass()->GetFName();
+	MediaPlateResourceName = FMediaPlateResource::StaticStruct()->GetFName();
 
 	FPropertyEditorModule& PropertyModule = FModuleManager::LoadModuleChecked<FPropertyEditorModule>("PropertyEditor");
 	PropertyModule.RegisterCustomClassLayout(MediaPlateName,
 		FOnGetDetailCustomizationInstance::CreateStatic(&FMediaPlateCustomization::MakeInstance));
 
+	PropertyModule.RegisterCustomPropertyTypeLayout(MediaPlateResourceName,
+		FOnGetPropertyTypeCustomizationInstance::CreateStatic(&FMediaPlateResourceCustomization::MakeInstance));
+
 	ISequencerModule& SequencerModule = FModuleManager::Get().LoadModuleChecked<ISequencerModule>("Sequencer");
 	TrackEditorBindingHandle = SequencerModule.RegisterPropertyTrackEditor<FMediaPlateTrackEditor>();
 
@@ -87,6 +91,7 @@ void FMediaPlateEditorModule::ShutdownModule()
 	// Unregister customizations.
 	FPropertyEditorModule& PropertyModule = FModuleManager::LoadModuleChecked<FPropertyEditorModule>("PropertyEditor");
 	PropertyModule.UnregisterCustomClassLayout(MediaPlateName);
+	PropertyModule.UnregisterCustomPropertyTypeLayout(MediaPlateResourceName);
 }
 
 void FMediaPlateEditorModule::Tick(float DeltaTime)
@@ -374,7 +379,7 @@ TSharedRef<FExtender> FMediaPlateEditorModule::ExtendLevelViewportContextMenuFor
 									}
 
 									const FScopedTransaction Transaction(LOCTEXT("ResetDefaultMats", "Reset Default Materials"));
-									
+
 									InMediaPlateActor->Modify();
 									InMediaPlateActor->UseDefaultMaterial();
 
@@ -402,8 +407,6 @@ TSharedRef<FExtender> FMediaPlateEditorModule::ExtendLevelViewportContextMenuFor
 				)
 			);
 		}
-
-
 	}
 
 	return Extender.ToSharedRef();
@@ -413,7 +416,7 @@ void FMediaPlateEditorModule::RegisterContextMenuExtender()
 {
 	// Extend the level viewport context menu to add an option to copy the object path.
 	LevelViewportContextMenuRemoteControlExtender = FLevelEditorModule::FLevelViewportMenuExtender_SelectedActors::CreateRaw(this, &FMediaPlateEditorModule::ExtendLevelViewportContextMenuForMediaPlate);
-	
+
 	FLevelEditorModule& LevelEditorModule = FModuleManager::Get().LoadModuleChecked<FLevelEditorModule>("LevelEditor");
 	TArray<FLevelEditorModule::FLevelViewportMenuExtender_SelectedActors>& MenuExtenders = LevelEditorModule.GetAllLevelViewportContextMenuExtenders();
 	MenuExtenders.Add(LevelViewportContextMenuRemoteControlExtender);
diff --git a/Engine/Plugins/Media/MediaPlate/Source/MediaPlateEditor/Private/MediaPlateResourceCustomization.cpp b/Engine/Plugins/Media/MediaPlate/Source/MediaPlateEditor/Private/MediaPlateResourceCustomization.cpp
new file mode 100644
--- /dev/null
+++ b/Engine/Plugins/Media/MediaPlate/Source/MediaPlateEditor/Private/MediaPlateResourceCustomization.cpp
@@ -0,0 +1,334 @@
+// Copyright Epic Games, Inc. All Rights Reserved.
+
+#include "MediaPlateResourceCustomization.h"
+#include "DetailLayoutBuilder.h"
+#include "DetailWidgetRow.h"
+#include "MediaPlateComponent.h"
+#include "MediaPlaylist.h"
+#include "MediaSource.h"
+#include "PropertyCustomizationHelpers.h"
+#include "ScopedTransaction.h"
+#include "Widgets/Input/SFilePathPicker.h"
+#include "Widgets/Input/SSegmentedControl.h"
+#include "Widgets/Layout/SBox.h"
+
+#define LOCTEXT_NAMESPACE "MediaPlateResourceCustomization"
+
+TSharedRef<IPropertyTypeCustomization> FMediaPlateResourceCustomization::MakeInstance()
+{
+	return MakeShared<FMediaPlateResourceCustomization>();
+}
+
+void FMediaPlateResourceCustomization::CustomizeHeader(TSharedRef<IPropertyHandle> InStructPropertyHandle, FDetailWidgetRow& InHeaderRow, IPropertyTypeCustomizationUtils& InStructCustomizationUtils)
+{
+	const TAttribute<EVisibility> MediaSourceFileVisibility(this, &FMediaPlateResourceCustomization::GetFileSelectorVisibility);
+	const TAttribute<EVisibility> MediaSourceAssetVisibility(this, &FMediaPlateResourceCustomization::GetAssetSelectorVisibility);
+	const TAttribute<EVisibility> MediaSourcePlaylistVisibility(this, &FMediaPlateResourceCustomization::GetPlaylistSelectorVisibility);
+
+	MediaPlateResourcePropertyHandle = InStructPropertyHandle;
+
+	void* StructRawData;
+	// Try accessing the raw value, so we make sure there's only one struct edited. Multiple Access edit is currently not supported
+	const FPropertyAccess::Result AccessResult = MediaPlateResourcePropertyHandle->GetValueData(StructRawData);
+
+	TSharedPtr<SWidget> ValueWidgetContent;
+
+	if (AccessResult == FPropertyAccess::Success)
+	{
+		ValueWidgetContent = SNew(SBox)
+		[
+			SNew(SVerticalBox)
+			+ SVerticalBox::Slot()
+			.HAlign(HAlign_Left)
+			[
+				SNew(SSegmentedControl<EMediaPlateResourceType>)
+					.Value(this, &FMediaPlateResourceCustomization::GetAssetType)
+					.OnValueChanged(this, &FMediaPlateResourceCustomization::OnAssetTypeChanged)
+
+				+ SSegmentedControl<EMediaPlateResourceType>::Slot(EMediaPlateResourceType::External)
+					.Text(LOCTEXT("File", "File"))
+					.ToolTip(LOCTEXT("File_ToolTip",
+						"Select this if you want to use a file path to a media file on disk."))
+
+				+ SSegmentedControl<EMediaPlateResourceType>::Slot(EMediaPlateResourceType::Asset)
+					.Text(LOCTEXT("Asset", "Asset"))
+					.ToolTip(LOCTEXT("Asset_ToolTip",
+						"Select this if you want to use a Media Source asset."))
+
+				+ SSegmentedControl<EMediaPlateResourceType>::Slot(EMediaPlateResourceType::Playlist)
+					.Text(LOCTEXT("Playlist", "Playlist"))
+					.ToolTip(LOCTEXT("Playlist_ToolTip",
+						"Select this if you want to use a Media Playlist asset."))
+			]
+			+ SVerticalBox::Slot()
+			[
+				SNew(SBox)
+				.Visibility(MediaSourceAssetVisibility)
+				.HAlign(HAlign_Fill)
+				[
+					SNew(SObjectPropertyEntryBox)
+					.AllowedClass(UMediaSource::StaticClass())
+					.ObjectPath(this, &FMediaPlateResourceCustomization::GetMediaAssetPath)
+					.OnObjectChanged(this, &FMediaPlateResourceCustomization::OnMediaAssetChanged)
+				]
+			]
+			+ SVerticalBox::Slot()
+			[
+				SNew(SBox)
+				.Visibility(MediaSourceFileVisibility)
+				.HAlign(HAlign_Fill)
+				[
+					SNew(SFilePathPicker)
+					.BrowseButtonImage(FAppStyle::GetBrush("PropertyWindow.Button_Ellipsis"))
+					.BrowseButtonStyle(FAppStyle::Get(), "HoverHintOnly")
+					.BrowseButtonToolTip(LOCTEXT("FileButtonToolTipText", "Choose a file from this computer"))
+					.BrowseTitle(LOCTEXT("PropertyEditorTitle", "File picker..."))
+					.FilePath(this, &FMediaPlateResourceCustomization::GetMediaPath)
+					.FileTypeFilter(TEXT("All files (*.*)|*.*"))
+					.OnPathPicked(this, &FMediaPlateResourceCustomization::OnMediaPathPicked)
+				]
+			]
+			+ SVerticalBox::Slot()
+			[
+				SNew(SBox)
+				.Visibility(MediaSourcePlaylistVisibility)
+				.HAlign(HAlign_Fill)
+				[
+					SNew(SObjectPropertyEntryBox)
+					.AllowedClass(UMediaPlaylist::StaticClass())
+					.ObjectPath(this, &FMediaPlateResourceCustomization::GetPlaylistPath)
+					.OnObjectChanged(this, &FMediaPlateResourceCustomization::OnPlaylistChanged)
+				]
+			]
+		];
+	}
+	else if (AccessResult == FPropertyAccess::MultipleValues)
+	{
+		ValueWidgetContent = SNew(STextBlock)
+			.Text(LOCTEXT("MultipleValues", "Multiple Selection"))
+			.ToolTipText(LOCTEXT("MultipleValues_ToolTip",
+			"Multiple Media Player Resource properties selected. Select a single property to edit it."))
+			.Font(IDetailLayoutBuilder::GetDetailFont());
+	}
+	else
+	{
+		ValueWidgetContent = SNew(STextBlock)
+			.Text(LOCTEXT("AccessError", "Error accessing property"))
+			.ToolTipText(LOCTEXT("AccessError_ToolTip",
+			"Error occurred while accessing Media Player Resource property."))
+			.Font(IDetailLayoutBuilder::GetDetailFont());
+	}
+
+	InHeaderRow.NameContent()
+	[
+		MediaPlateResourcePropertyHandle->CreatePropertyNameWidget()
+	]
+	.ValueContent()
+	[
+		ValueWidgetContent.ToSharedRef()
+	];
+}
+
+void FMediaPlateResourceCustomization::CustomizeChildren(TSharedRef<IPropertyHandle> InPropertyHandle, IDetailChildrenBuilder& InChildBuilder, IPropertyTypeCustomizationUtils& InCustomizationUtils)
+{
+}
+
+EMediaPlateResourceType FMediaPlateResourceCustomization::GetAssetType() const
+{
+	constexpr EMediaPlateResourceType DefaultType = EMediaPlateResourceType::Asset;
+
+	if (const FMediaPlateResource* MediaPlateResource = GetMediaPlateSourcePtr())
+	{
+		return MediaPlateResource->GetResourceType();
+	}
+
+	return DefaultType;
+}
+
+FMediaPlateResource* FMediaPlateResourceCustomization::GetMediaPlateSourcePtr() const
+{
+	if (!MediaPlateResourcePropertyHandle)
+	{
+		return nullptr;
+	}
+
+	void* StructRawData;
+	const FPropertyAccess::Result AccessResult = MediaPlateResourcePropertyHandle->GetValueData(StructRawData);
+	if (AccessResult == FPropertyAccess::Success)
+	{
+		return static_cast<FMediaPlateResource*>(StructRawData);
+	}
+
+	return nullptr;
+}
+
+UObject* FMediaPlateResourceCustomization::GetMediaPlateResourceOwner() const
+{
+	TArray<UObject*> OuterObjects;
+	MediaPlateResourcePropertyHandle->GetOuterObjects(OuterObjects);
+
+	if (OuterObjects.Num() != 1)
+	{
+		return nullptr;
+	}
+
+	if (UObject* Outer = OuterObjects[0])
+	{
+		return Outer;
+	}
+
+	return nullptr;
+}
+
+void FMediaPlateResourceCustomization::OnAssetTypeChanged(EMediaPlateResourceType InMediaSourceType)
+{
+	if (GetAssetType() != InMediaSourceType)
+	{
+		// Updating value and notifying, in case the user is switching between already specified options
+		if (FMediaPlateResource* MediaPlateResource = GetMediaPlateSourcePtr())
+		{
+			const FScopedTransaction Transaction(LOCTEXT("OnMediaSourceTypeChanged", "Media Source type changed"));
+
+			if (UObject* MediaPlateResourceOwner = GetMediaPlateResourceOwner())
+			{
+				MediaPlateResourceOwner->Modify();
+				MediaPlateResource->Modify();
+			}
+
+			MediaPlateResource->SetResourceType(InMediaSourceType);
+			MediaPlateResourcePropertyHandle->NotifyPostChange(EPropertyChangeType::ValueSet);
+		}
+	}
+}
+
+FString FMediaPlateResourceCustomization::GetMediaAssetPath() const
+{
+	FString MediaAssetPath;
+
+	if (const FMediaPlateResource* MediaPlateResource = GetMediaPlateSourcePtr())
+	{
+		return MediaPlateResource->GetMediaAsset()->GetPathName();
+	}
+
+	return MediaAssetPath;
+}
+
+FString FMediaPlateResourceCustomization::GetPlaylistPath() const
+{
+	FString PlaylistAssetPath;
+
+	if (const FMediaPlateResource* MediaPlateResource = GetMediaPlateSourcePtr())
+	{
+		if (const UMediaPlaylist* Playlist = MediaPlateResource->GetSourcePlaylist())
+		{
+			PlaylistAssetPath = Playlist->GetPathName();
+		}
+	}
+
+	return PlaylistAssetPath;
+}
+
+FString FMediaPlateResourceCustomization::GetMediaPath() const
+{
+	FString ExternalMediaPath;
+
+	if (const FMediaPlateResource* MediaPlateResource = GetMediaPlateSourcePtr())
+	{
+		ExternalMediaPath = MediaPlateResource->GetExternalMediaPath();
+	}
+
+	return ExternalMediaPath;
+}
+
+void FMediaPlateResourceCustomization::OnMediaAssetChanged(const FAssetData& InAssetData)
+{
+	if (!MediaPlateResourcePropertyHandle)
+	{
+		return;
+	}
+
+	if (const UMediaSource* MediaAsset = Cast<UMediaSource>(InAssetData.GetAsset()))
+	{
+		if (FMediaPlateResource* MediaPlateResource = GetMediaPlateSourcePtr())
+		{
+			if (UObject* MediaPlateResourceOwner = GetMediaPlateResourceOwner())
+			{
+				const FScopedTransaction Transaction(LOCTEXT("OnMediaSourceAssetChanged", "Media Source asset changed"));
+				MediaPlateResourceOwner->Modify();
+				MediaPlateResource->Modify();
+
+				MediaPlateResource->SelectAsset(MediaAsset, MediaPlateResourceOwner);
+				MediaPlateResourcePropertyHandle->NotifyPostChange(EPropertyChangeType::ValueSet);
+			}
+		}
+	}
+}
+
+void FMediaPlateResourceCustomization::OnMediaPathPicked(const FString& InPickedPath)
+{
+	if (!MediaPlateResourcePropertyHandle)
+	{
+		return;
+	}
+
+	FMediaPlateResource* MediaPlateResource = GetMediaPlateSourcePtr();
+
+	if (!MediaPlateResource)
+	{
+		return;
+	}
+
+	if (!InPickedPath.IsEmpty() && InPickedPath != MediaPlateResource->GetExternalMediaPath())
+	{
+		const FScopedTransaction Transaction(LOCTEXT("OnMediaExternalPathChanged", "Media external file path changed"));
+		if (UObject* MediaPlateResourceOwner = GetMediaPlateResourceOwner())
+		{
+			MediaPlateResourceOwner->Modify();
+			MediaPlateResource->Modify();
+
+			MediaPlateResource->LoadExternalMedia(InPickedPath, MediaPlateResourceOwner);
+			MediaPlateResourcePropertyHandle->NotifyPostChange(EPropertyChangeType::ValueSet);
+		}
+	}
+}
+
+void FMediaPlateResourceCustomization::OnPlaylistChanged(const FAssetData& InAssetData)
+{
+	if (!MediaPlateResourcePropertyHandle)
+	{
+		return;
+	}
+
+	const UMediaPlaylist* Playlist = Cast<UMediaPlaylist>(InAssetData.GetAsset());
+	if (FMediaPlateResource* MediaPlateResource = GetMediaPlateSourcePtr())
+	{
+		const FScopedTransaction Transaction(LOCTEXT("OnPlaylistChanged", "Media Playlist changed"));
+		if (UObject* MediaPlateResourceOwner = GetMediaPlateResourceOwner())
+		{
+			MediaPlateResourceOwner->Modify();
+		}
+
+		// No need to call MediaPlateResource->Modify():
+		// ActivePlaylist property will be overwritten, not modified
+
+		MediaPlateResource->SelectPlaylist(Playlist);
+		MediaPlateResourcePropertyHandle->NotifyPostChange(EPropertyChangeType::ValueSet);
+	}
+}
+
+EVisibility FMediaPlateResourceCustomization::GetAssetSelectorVisibility() const
+{
+	return GetAssetType() == EMediaPlateResourceType::Asset ? EVisibility::Visible : EVisibility::Collapsed;
+}
+
+EVisibility FMediaPlateResourceCustomization::GetFileSelectorVisibility() const
+{
+	return GetAssetType() == EMediaPlateResourceType::External ? EVisibility::Visible : EVisibility::Collapsed;
+}
+
+EVisibility FMediaPlateResourceCustomization::GetPlaylistSelectorVisibility() const
+{
+	return GetAssetType() == EMediaPlateResourceType::Playlist ? EVisibility::Visible : EVisibility::Collapsed;
+}
+
+#undef LOCTEXT_NAMESPACE
diff --git a/Engine/Plugins/Media/MediaPlate/Source/MediaPlateEditor/Private/MediaPlateResourceCustomization.h b/Engine/Plugins/Media/MediaPlate/Source/MediaPlateEditor/Private/MediaPlateResourceCustomization.h
new file mode 100644
--- /dev/null
+++ b/Engine/Plugins/Media/MediaPlate/Source/MediaPlateEditor/Private/MediaPlateResourceCustomization.h
@@ -0,0 +1,47 @@
+// Copyright Epic Games, Inc. All Rights Reserved.
+
+#pragma once
+
+#include "IPropertyTypeCustomization.h"
+#include "MediaPlateComponent.h"
+
+class IPropertyHandle;
+class UMediaSource;
+struct EVisibility;
+struct FMediaPlateResource;
+
+class FMediaPlateResourceCustomization : public IPropertyTypeCustomization
+{
+public:
+	static TSharedRef<IPropertyTypeCustomization> MakeInstance();
+
+	//~ Begin IPropertyTypeCustomization
+	virtual void CustomizeHeader(TSharedRef<IPropertyHandle> InStructPropertyHandle
+		, FDetailWidgetRow& InHeaderRow
+		, IPropertyTypeCustomizationUtils& InStructCustomizationUtils) override;
+
+	virtual void CustomizeChildren(TSharedRef<IPropertyHandle> InPropertyHandle
+		, IDetailChildrenBuilder& InChildBuilder
+		, IPropertyTypeCustomizationUtils& InCustomizationUtils) override;
+	//~ End IPropertyTypeCustomization
+
+private:
+	EMediaPlateResourceType GetAssetType() const;
+
+	FString GetMediaAssetPath() const;
+	FString GetPlaylistPath() const;
+	FString GetMediaPath() const;
+
+	void OnAssetTypeChanged(EMediaPlateResourceType InMediaSourceType);
+	void OnPlaylistChanged(const FAssetData& InAssetData);
+	void OnMediaAssetChanged(const FAssetData& InAssetData);
+	void OnMediaPathPicked(const FString& InPickedPath);
+
+	EVisibility GetAssetSelectorVisibility() const;
+	EVisibility GetFileSelectorVisibility() const;
+	EVisibility GetPlaylistSelectorVisibility() const;
+	FMediaPlateResource* GetMediaPlateSourcePtr() const;
+	UObject* GetMediaPlateResourceOwner() const;
+
+	TSharedPtr<IPropertyHandle> MediaPlateResourcePropertyHandle;
+};
diff --git a/Engine/Plugins/Media/MediaPlate/Source/MediaPlateEditor/Private/Sequencer/MediaPlateTrackEditor.cpp b/Engine/Plugins/Media/MediaPlate/Source/MediaPlateEditor/Private/Sequencer/MediaPlateTrackEditor.cpp
--- a/Engine/Plugins/Media/MediaPlate/Source/MediaPlateEditor/Private/Sequencer/MediaPlateTrackEditor.cpp
+++ b/Engine/Plugins/Media/MediaPlate/Source/MediaPlateEditor/Private/Sequencer/MediaPlateTrackEditor.cpp
@@ -152,7 +152,7 @@ void FMediaPlateTrackEditor::AddTrackForComponent(UMediaPlateComponent* Componen
 		MediaTrack->SetDisplayName(LOCTEXT("MediaTrackName", "Media"));
 
 		// Populate track.
-		UMediaPlaylist* Playlist = Component->MediaPlaylist;
+		UMediaPlaylist* Playlist = Component->GetMediaPlaylist();
 		if (Playlist != nullptr)
 		{
 			for (int32 Index = 0; Index < Playlist->Num(); ++Index)
diff --git a/Engine/Plugins/Media/MediaPlate/Source/MediaPlateEditor/Private/Toolkits/MediaPlateEditorToolkit.cpp b/Engine/Plugins/Media/MediaPlate/Source/MediaPlateEditor/Private/Toolkits/MediaPlateEditorToolkit.cpp
--- a/Engine/Plugins/Media/MediaPlate/Source/MediaPlateEditor/Private/Toolkits/MediaPlateEditorToolkit.cpp
+++ b/Engine/Plugins/Media/MediaPlate/Source/MediaPlateEditor/Private/Toolkits/MediaPlateEditorToolkit.cpp
@@ -271,8 +271,8 @@ void FMediaPlateEditorToolkit::BindCommands()
 		Commands.NextMedia,
 		FExecuteAction::CreateLambda([this]{ MediaPlate->Next(); }),
 		FCanExecuteAction::CreateLambda([this]{
-			return (MediaPlate->MediaPlaylist != nullptr) &&
-				(MediaPlate->MediaPlaylist->Num() > 1);
+			return (MediaPlate->GetMediaPlaylist() != nullptr) &&
+				(MediaPlate->GetMediaPlaylist()->Num() > 1);
 		})
 	);
 
@@ -304,8 +304,8 @@ void FMediaPlateEditorToolkit::BindCommands()
 		Commands.PreviousMedia,
 		FExecuteAction::CreateLambda([this]{ MediaPlate->Previous(); }),
 		FCanExecuteAction::CreateLambda([this]{
-			return (MediaPlate->MediaPlaylist != nullptr) &&
-					(MediaPlate->MediaPlaylist->Num() > 1);
+			return (MediaPlate->GetMediaPlaylist() != nullptr) &&
+					(MediaPlate->GetMediaPlaylist()->Num() > 1);
 		})
 	);
 
diff --git a/Engine/Plugins/Media/MediaPlate/Source/MediaPlateEditor/Private/Widgets/SMediaPlateEditorPlaylist.cpp b/Engine/Plugins/Media/MediaPlate/Source/MediaPlateEditor/Private/Widgets/SMediaPlateEditorPlaylist.cpp
--- a/Engine/Plugins/Media/MediaPlate/Source/MediaPlateEditor/Private/Widgets/SMediaPlateEditorPlaylist.cpp
+++ b/Engine/Plugins/Media/MediaPlate/Source/MediaPlateEditor/Private/Widgets/SMediaPlateEditorPlaylist.cpp
@@ -10,7 +10,7 @@ void SMediaPlateEditorPlaylist::Construct(const FArguments& InArgs, UMediaPlateC
 {
 	MediaPlate = &InMediaPlate;
 
-	SMediaPlaylistEditorTracks::Construct(SMediaPlaylistEditorTracks::FArguments(), MediaPlate->MediaPlaylist, InStyle);
+	SMediaPlaylistEditorTracks::Construct(SMediaPlaylistEditorTracks::FArguments(), MediaPlate->GetMediaPlaylist(), InStyle);
 }
 
 void SMediaPlateEditorPlaylist::Tick(const FGeometry& AllottedGeometry, const double InCurrentTime, const float InDeltaTime)
@@ -22,10 +22,10 @@ void SMediaPlateEditorPlaylist::Tick(const FGeometry& AllottedGeometry, const do
 	{
 		// Did the playlist change?
 		UMediaPlaylist* MediaPlaylist = MediaPlaylistPtr.Get();
-		if (MediaPlaylist != MediaPlate->MediaPlaylist)
+		if (MediaPlaylist != MediaPlate->GetMediaPlaylist())
 		{
 			// Refresh.
-			MediaPlaylistPtr = MediaPlate->MediaPlaylist;
+			MediaPlaylistPtr = MediaPlate->GetMediaPlaylist();
 			RefreshPlaylist();
 		}
 	}
diff --git a/Engine/Plugins/Media/MediaPlate/Source/MediaPlateEditor/Public/MediaPlateEditorModule.h b/Engine/Plugins/Media/MediaPlate/Source/MediaPlateEditor/Public/MediaPlateEditorModule.h
--- a/Engine/Plugins/Media/MediaPlate/Source/MediaPlateEditor/Public/MediaPlateEditorModule.h
+++ b/Engine/Plugins/Media/MediaPlate/Source/MediaPlateEditor/Public/MediaPlateEditorModule.h
@@ -50,8 +50,10 @@ public:
 	
 private:
 
-	/** Customization name to avoid reusing staticstruct during shutdown. */
+	/** Customization names to avoid reusing staticstruct during shutdown. */
 	FName MediaPlateName;
+	FName MediaPlateResourceName;
+
 	/** Holds all the media plates that are playing. */
 	TArray<TWeakObjectPtr<UMediaPlateComponent>> ActiveMediaPlates;
 
