[Media Plate] Analytics Event crashing when creating Media Plate Actor in Debug

diff --git a/Engine/Plugins/Media/MediaPlate/Source/MediaPlateEditor/Private/MediaPlateEditorModule.cpp b/Engine/Plugins/Media/MediaPlate/Source/MediaPlateEditor/Private/MediaPlateEditorModule.cpp
--- a/Engine/Plugins/Media/MediaPlate/Source/MediaPlateEditor/Private/MediaPlateEditorModule.cpp
+++ b/Engine/Plugins/Media/MediaPlate/Source/MediaPlateEditor/Private/MediaPlateEditorModule.cpp
@@ -443,7 +443,7 @@ void FMediaPlateEditorModule::OnLevelActorAdded(AActor* InActor)
 		return;
 	}
 
-	if (InActor->GetClass() == AMediaPlate::StaticClass())
+	if (FEngineAnalytics::IsAvailable() && InActor->GetClass() == AMediaPlate::StaticClass())
 	{
 		FEngineAnalytics::GetProvider().RecordEvent(TEXT("Editor.Usage.MediaPlate.AddMediaPlateActorToLevel"));
 	}
diff --git a/Engine/Plugins/Media/MediaPlate/Source/MediaPlateEditor/Private/Sequencer/MediaPlateTrackEditor.cpp b/Engine/Plugins/Media/MediaPlate/Source/MediaPlateEditor/Private/Sequencer/MediaPlateTrackEditor.cpp
--- a/Engine/Plugins/Media/MediaPlate/Source/MediaPlateEditor/Private/Sequencer/MediaPlateTrackEditor.cpp
+++ b/Engine/Plugins/Media/MediaPlate/Source/MediaPlateEditor/Private/Sequencer/MediaPlateTrackEditor.cpp
@@ -121,7 +121,7 @@ void FMediaPlateTrackEditor::HandleActorAdded(AActor* Actor, FGuid TargetObjectG
 {
 	if (Actor)
 	{
-		if (Actor->GetClass() == AMediaPlate::StaticClass())
+		if (FEngineAnalytics::IsAvailable() && Actor->GetClass() == AMediaPlate::StaticClass())
 		{
 			FEngineAnalytics::GetProvider().RecordEvent(TEXT("Editor.Usage.MediaPlate.AddMediaPlateActorToSequencer"));
 		}
