[Media Plate] Add missing nullptr check against GEditor and GEngine

diff --git a/Engine/Plugins/Media/MediaPlate/Source/MediaPlateEditor/Private/MediaPlateCustomization.cpp b/Engine/Plugins/Media/MediaPlate/Source/MediaPlateEditor/Private/MediaPlateCustomization.cpp
--- a/Engine/Plugins/Media/MediaPlate/Source/MediaPlateEditor/Private/MediaPlateCustomization.cpp
+++ b/Engine/Plugins/Media/MediaPlate/Source/MediaPlateEditor/Private/MediaPlateCustomization.cpp
@@ -910,7 +910,7 @@ FReply FMediaPlateCustomization::OnOpenMediaPlate()
 	}
 
 	// Open the editor.
-	if (AssetArray.Num() > 0)
+	if (GEditor && AssetArray.Num() > 0)
 	{
 		GEditor->GetEditorSubsystem<UAssetEditorSubsystem>()->OpenEditorForAssets(AssetArray);
 	}
diff --git a/Engine/Plugins/Media/MediaPlate/Source/MediaPlateEditor/Private/Toolkits/MediaPlateEditorToolkit.cpp b/Engine/Plugins/Media/MediaPlate/Source/MediaPlateEditor/Private/Toolkits/MediaPlateEditorToolkit.cpp
--- a/Engine/Plugins/Media/MediaPlate/Source/MediaPlateEditor/Private/Toolkits/MediaPlateEditorToolkit.cpp
+++ b/Engine/Plugins/Media/MediaPlate/Source/MediaPlateEditor/Private/Toolkits/MediaPlateEditorToolkit.cpp
@@ -44,9 +44,16 @@ FMediaPlateEditorToolkit::~FMediaPlateEditorToolkit()
 {
 	FReimportManager::Instance()->OnPreReimport().RemoveAll(this);
 	FReimportManager::Instance()->OnPostReimport().RemoveAll(this);
-	
-	GEngine->OnLevelActorDeleted().RemoveAll(this);
-	GEditor->UnregisterForUndo(this);
+
+	if (GEngine)
+	{
+		GEngine->OnLevelActorDeleted().RemoveAll(this);
+	}
+
+	if (GEditor)
+	{
+		GEditor->UnregisterForUndo(this);
+	}
 }
 
 /* FMediaPlateEditorToolkit interface
@@ -63,7 +70,11 @@ void FMediaPlateEditorToolkit::Initialize(UMediaPlateComponent* InMediaPlate, co
 
 	// support undo/redo
 	MediaPlate->SetFlags(RF_Transactional);
-	GEditor->RegisterForUndo(this);
+
+	if (GEditor)
+	{
+		GEditor->RegisterForUndo(this);
+	}
 
 	BindCommands();
 
@@ -136,7 +147,10 @@ void FMediaPlateEditorToolkit::Initialize(UMediaPlateComponent* InMediaPlate, co
 		EditorModule->MediaPlateStartedPlayback(MediaPlate);
 	}
 
-	GEngine->OnLevelActorDeleted().AddSP(this, &FMediaPlateEditorToolkit::OnActorDeleted);
+	if (GEngine)
+	{
+		GEngine->OnLevelActorDeleted().AddSP(this, &FMediaPlateEditorToolkit::OnActorDeleted);
+	}
 }
 
 /* FAssetEditorToolkit interface
diff --git a/Engine/Plugins/Media/MediaPlate/Source/MediaPlateEditor/Private/Widgets/SMediaPlateEditorMaterial.cpp b/Engine/Plugins/Media/MediaPlate/Source/MediaPlateEditor/Private/Widgets/SMediaPlateEditorMaterial.cpp
--- a/Engine/Plugins/Media/MediaPlate/Source/MediaPlateEditor/Private/Widgets/SMediaPlateEditorMaterial.cpp
+++ b/Engine/Plugins/Media/MediaPlate/Source/MediaPlateEditor/Private/Widgets/SMediaPlateEditorMaterial.cpp
@@ -77,7 +77,11 @@ FReply SMediaPlateEditorMaterial::OnBrowseMaterial() const
 			{
 				TArray<UObject*> ObjectsToSync;
 				ObjectsToSync.Add(BaseMaterial);
-				GEditor->SyncBrowserToObjects(ObjectsToSync);
+
+				if (GEditor)
+				{
+					GEditor->SyncBrowserToObjects(ObjectsToSync);
+				}
 			}
 		}
 	}
