MediaPlate: Fixed possible crash in the media plate customization attribute lambda.

Replaced all lambda captures by this with a safe weak pointer to self.

#rnx

diff --git a/Engine/Plugins/Media/MediaPlate/Source/MediaPlateEditor/Private/MediaPlateCustomization.cpp b/Engine/Plugins/Media/MediaPlate/Source/MediaPlateEditor/Private/MediaPlateCustomization.cpp
--- a/Engine/Plugins/Media/MediaPlate/Source/MediaPlateEditor/Private/MediaPlateCustomization.cpp
+++ b/Engine/Plugins/Media/MediaPlate/Source/MediaPlateEditor/Private/MediaPlateCustomization.cpp
@@ -45,6 +45,8 @@ FMediaPlateCustomization::~FMediaPlateCustomization()
 
 void FMediaPlateCustomization::CustomizeDetails(IDetailLayoutBuilder& DetailBuilder)
 {
+	TWeakPtr<FMediaPlateCustomization> WeakSelf = StaticCastWeakPtr<FMediaPlateCustomization>(AsWeak());
+
 	// Is this the media plate editor window?
 	bool bIsMediaPlateWindow = false;
 
@@ -114,21 +116,30 @@ void FMediaPlateCustomization::CustomizeDetails(IDetailLayoutBuilder& DetailBuil
 				[
 					SNew(SButton)
 						.VAlign(VAlign_Center)
-						.IsEnabled_Lambda([this]
+						.IsEnabled_Lambda([WeakSelf]
 						{
-							if (const UMediaPlayer* MediaPlayer = GetMediaPlayer())
+							const TSharedPtr<FMediaPlateCustomization> Self = WeakSelf.Pin();
+							if (Self.IsValid())
 							{
-								return MediaPlayer->IsReady() &&
-									MediaPlayer->SupportsSeeking() &&
-									MediaPlayer->GetTime() > FTimespan::Zero();
+								if (const UMediaPlayer* MediaPlayer = Self->GetMediaPlayer())
+								{
+									return MediaPlayer->IsReady() &&
+										MediaPlayer->SupportsSeeking() &&
+										MediaPlayer->GetTime() > FTimespan::Zero();
+								}
 							}
 
 							return false;
 						})
-						.OnClicked_Lambda([this]() -> FReply
+						.OnClicked_Lambda([WeakSelf]() -> FReply
 						{
-							OnButtonEvent(EMediaPlateEventState::Rewind);
-							return FReply::Handled();
+							const TSharedPtr<FMediaPlateCustomization> Self = WeakSelf.Pin();
+							if (Self.IsValid())
+							{
+								Self->OnButtonEvent(EMediaPlateEventState::Rewind);
+								return FReply::Handled();
+							}
+							return FReply::Unhandled();
 						})
 						[
 							SNew(SImage)
@@ -145,19 +156,28 @@ void FMediaPlateCustomization::CustomizeDetails(IDetailLayoutBuilder& DetailBuil
 				[
 					SNew(SButton)
 						.VAlign(VAlign_Center)
-						.IsEnabled_Lambda([this]
+						.IsEnabled_Lambda([WeakSelf]
 						{
-							if (UMediaPlayer* MediaPlayer = GetMediaPlayer())
+							const TSharedPtr<FMediaPlateCustomization> Self = WeakSelf.Pin();
+							if (Self.IsValid())
 							{
-								return MediaPlayer->IsReady() && MediaPlayer->SupportsRate(UMediaPlateComponent::GetReverseRate(MediaPlayer), false);
+								if (UMediaPlayer* MediaPlayer = Self->GetMediaPlayer())
+								{
+									return MediaPlayer->IsReady() && MediaPlayer->SupportsRate(UMediaPlateComponent::GetReverseRate(MediaPlayer), false);
+								}
 							}
 
 							return false;
 						})
-						.OnClicked_Lambda([this]() -> FReply
+						.OnClicked_Lambda([WeakSelf]() -> FReply
 						{
-							OnButtonEvent(EMediaPlateEventState::Reverse);
-							return FReply::Handled();
+							const TSharedPtr<FMediaPlateCustomization> Self = WeakSelf.Pin();
+							if (Self.IsValid())
+							{
+								Self->OnButtonEvent(EMediaPlateEventState::Reverse);
+								return FReply::Handled();
+							}
+							return FReply::Unhandled();
 						})
 						[
 							SNew(SImage)
@@ -174,30 +194,44 @@ void FMediaPlateCustomization::CustomizeDetails(IDetailLayoutBuilder& DetailBuil
 				[
 					SNew(SButton)
 						.VAlign(VAlign_Center)
-						.IsEnabled_Lambda([this]
+						.IsEnabled_Lambda([WeakSelf]
 						{
-							// Is the player paused or fast forwarding/rewinding?
-							if (const UMediaPlayer* MediaPlayer = GetMediaPlayer())
+							const TSharedPtr<FMediaPlateCustomization> Self = WeakSelf.Pin();
+							if (Self.IsValid())
 							{
-								return MediaPlayer->IsReady()
-									&& (!MediaPlayer->IsPlaying() || (MediaPlayer->GetRate() != 1.0f));
+								// Is the player paused or fast forwarding/rewinding?
+								if (const UMediaPlayer* MediaPlayer = Self->GetMediaPlayer())
+								{
+									return MediaPlayer->IsReady()
+										&& (!MediaPlayer->IsPlaying() || (MediaPlayer->GetRate() != 1.0f));
+								}
 							}
 
 							return false;
 						})
-						.Visibility_Lambda([this]
+						.Visibility_Lambda([WeakSelf]
 						{
-							if (const UMediaPlayer* MediaPlayer = GetMediaPlayer())
+							const TSharedPtr<FMediaPlateCustomization> Self = WeakSelf.Pin();
+							if (Self.IsValid())
 							{
-								return MediaPlayer->IsPlaying() ? EVisibility::Collapsed : EVisibility::Visible;
+								if (const UMediaPlayer* MediaPlayer = Self->GetMediaPlayer())
+								{
+									return MediaPlayer->IsPlaying() ? EVisibility::Collapsed : EVisibility::Visible;
+								}
 							}
 
 							return EVisibility::Visible;
 						})
-						.OnClicked_Lambda([this]() -> FReply
+						.OnClicked_Lambda([WeakSelf]() -> FReply
 						{
-							OnButtonEvent(EMediaPlateEventState::Play);
-							return FReply::Handled();
+							const TSharedPtr<FMediaPlateCustomization> Self = WeakSelf.Pin();
+							if (Self.IsValid())
+							{
+								Self->OnButtonEvent(EMediaPlateEventState::Play);
+								return FReply::Handled();
+							}
+
+							return FReply::Unhandled();
 						})
 						[
 							SNew(SImage)
@@ -214,29 +248,42 @@ void FMediaPlateCustomization::CustomizeDetails(IDetailLayoutBuilder& DetailBuil
 				[
 					SNew(SButton)
 						.VAlign(VAlign_Center)
-						.IsEnabled_Lambda([this]
+						.IsEnabled_Lambda([WeakSelf]
 						{
-							if (const UMediaPlayer* MediaPlayer = GetMediaPlayer())
+							const TSharedPtr<FMediaPlateCustomization> Self = WeakSelf.Pin();
+							if (Self.IsValid())
 							{
-								return MediaPlayer->CanPause() && !MediaPlayer->IsPaused();
+								if (const UMediaPlayer* MediaPlayer = Self->GetMediaPlayer())
+								{
+									return MediaPlayer->CanPause() && !MediaPlayer->IsPaused();
+								}
 							}
 
 							return false;
 						})
-						.Visibility_Lambda([this]
+						.Visibility_Lambda([WeakSelf]
 						{
-							if (const UMediaPlayer* MediaPlayer = GetMediaPlayer())
+							const TSharedPtr<FMediaPlateCustomization> Self = WeakSelf.Pin();
+							if (Self.IsValid())
 							{
-								const bool bIsVisible = MediaPlayer->CanPause() && !MediaPlayer->IsPaused();
-								return bIsVisible ? EVisibility::Visible : EVisibility::Collapsed;
+								if (const UMediaPlayer* MediaPlayer = Self->GetMediaPlayer())
+								{
+									const bool bIsVisible = MediaPlayer->CanPause() && !MediaPlayer->IsPaused();
+									return bIsVisible ? EVisibility::Visible : EVisibility::Collapsed;
+								}
 							}
 
 							return EVisibility::Collapsed;
 						})
-						.OnClicked_Lambda([this]() -> FReply
+						.OnClicked_Lambda([WeakSelf]() -> FReply
 						{
-							OnButtonEvent(EMediaPlateEventState::Pause);
-							return FReply::Handled();
+							const TSharedPtr<FMediaPlateCustomization> Self = WeakSelf.Pin();
+							if (Self.IsValid())
+							{
+								Self->OnButtonEvent(EMediaPlateEventState::Pause);
+								return FReply::Handled();
+							}
+							return FReply::Unhandled();
 						})
 						[
 							SNew(SImage)
@@ -253,19 +300,28 @@ void FMediaPlateCustomization::CustomizeDetails(IDetailLayoutBuilder& DetailBuil
 				[
 					SNew(SButton)
 					.VAlign(VAlign_Center)
-					.IsEnabled_Lambda([this]
+					.IsEnabled_Lambda([WeakSelf]
 					{
-						if (UMediaPlayer* MediaPlayer = GetMediaPlayer())
+						const TSharedPtr<FMediaPlateCustomization> Self = WeakSelf.Pin();
+						if (Self.IsValid())
 						{
-							return MediaPlayer->IsReady() && MediaPlayer->SupportsRate(UMediaPlateComponent::GetForwardRate(MediaPlayer), false);
+							if (UMediaPlayer* MediaPlayer = Self->GetMediaPlayer())
+							{
+								return MediaPlayer->IsReady() && MediaPlayer->SupportsRate(UMediaPlateComponent::GetForwardRate(MediaPlayer), false);
+							}
 						}
 
 						return false;
 					})
-					.OnClicked_Lambda([this]() -> FReply
+					.OnClicked_Lambda([WeakSelf]() -> FReply
 					{
-						OnButtonEvent(EMediaPlateEventState::Forward);
-						return FReply::Handled();
+						const TSharedPtr<FMediaPlateCustomization> Self = WeakSelf.Pin();
+						if (Self.IsValid())
+						{
+							Self->OnButtonEvent(EMediaPlateEventState::Forward);
+							return FReply::Handled();
+						}
+						return FReply::Unhandled();
 					})
 					[
 						SNew(SImage)
@@ -282,10 +338,15 @@ void FMediaPlateCustomization::CustomizeDetails(IDetailLayoutBuilder& DetailBuil
 				[
 					SNew(SButton)
 						.VAlign(VAlign_Center)
-						.OnClicked_Lambda([this]() -> FReply
+						.OnClicked_Lambda([WeakSelf]() -> FReply
 						{
-							OnButtonEvent(EMediaPlateEventState::Open);
-							return FReply::Handled();
+							const TSharedPtr<FMediaPlateCustomization> Self = WeakSelf.Pin();
+							if (Self.IsValid())
+							{
+								Self->OnButtonEvent(EMediaPlateEventState::Open);
+								return FReply::Handled();
+							}
+							return FReply::Unhandled();
 						})
 						[
 							SNew(SImage)
@@ -302,19 +363,28 @@ void FMediaPlateCustomization::CustomizeDetails(IDetailLayoutBuilder& DetailBuil
 				[
 					SNew(SButton)
 						.VAlign(VAlign_Center)
-						.IsEnabled_Lambda([this]
+						.IsEnabled_Lambda([WeakSelf]
 						{
-							if (UMediaPlayer* MediaPlayer = GetMediaPlayer())
+							const TSharedPtr<FMediaPlateCustomization> Self = WeakSelf.Pin();
+							if (Self.IsValid())
 							{
-								return !MediaPlayer->GetUrl().IsEmpty();
+								if (UMediaPlayer* MediaPlayer = Self->GetMediaPlayer())
+								{
+									return !MediaPlayer->GetUrl().IsEmpty();
+								}
 							}
 
 							return false;
 						})
-						.OnClicked_Lambda([this]() -> FReply
+						.OnClicked_Lambda([WeakSelf]() -> FReply
 						{
-							OnButtonEvent(EMediaPlateEventState::Close);
-							return FReply::Handled();
+							const TSharedPtr<FMediaPlateCustomization> Self = WeakSelf.Pin();
+							if (Self.IsValid())
+							{
+								Self->OnButtonEvent(EMediaPlateEventState::Close);
+								return FReply::Handled();
+							}
+							return FReply::Unhandled();
 						})
 						[
 							SNew(SImage)
@@ -369,13 +439,16 @@ void FMediaPlateCustomization::CustomizeDetails(IDetailLayoutBuilder& DetailBuil
 
 void FMediaPlateCustomization::AddMeshCustomization(IDetailCategoryBuilder& InParentCategory)
 {
+	TWeakPtr<FMediaPlateCustomization> WeakSelf = StaticCastWeakPtr<FMediaPlateCustomization>(AsWeak());
+
 	// Add radio buttons for mesh type.
 	InParentCategory.AddCustomRow(FText::FromString("Mesh Selection"))
 	[
 		SNew(SSegmentedControl<EMediaTextureVisibleMipsTiles>)
-			.Value_Lambda([this]()
+			.Value_Lambda([WeakSelf]()
 			{
-				return MeshMode;
+				const TSharedPtr<FMediaPlateCustomization> Self = WeakSelf.Pin();
+				return Self.IsValid() ? Self->MeshMode : EMediaTextureVisibleMipsTiles::None;
 			})
 			.OnValueChanged(this, &FMediaPlateCustomization::SetMeshMode)
 
