MediaPlate: Tentative fix to playback events for the top bar menu not transacting.

This replicates the logic present in the regular media plate customization.

diff --git a/Engine/Plugins/Media/MediaPlate/Source/MediaPlate/Private/MediaPlateComponent.cpp b/Engine/Plugins/Media/MediaPlate/Source/MediaPlate/Private/MediaPlateComponent.cpp
--- a/Engine/Plugins/Media/MediaPlate/Source/MediaPlate/Private/MediaPlateComponent.cpp
+++ b/Engine/Plugins/Media/MediaPlate/Source/MediaPlate/Private/MediaPlateComponent.cpp
@@ -1482,6 +1482,22 @@ void UMediaPlateComponent::SwitchStates(EMediaPlateEventState State)
 			}
 		}
 		break;
+	case EMediaPlateEventState::Next:
+		{
+			if (MediaPlayer != nullptr)
+			{
+				MediaPlayer->Next();
+			}
+		}
+		break;
+	case EMediaPlateEventState::Previous:
+		{
+			if (MediaPlayer != nullptr)
+			{
+				MediaPlayer->Previous();
+			}
+		}
+		break;
 	case EMediaPlateEventState::MAX:
 	default:
 		checkNoEntry();
diff --git a/Engine/Plugins/Media/MediaPlate/Source/MediaPlate/Public/MediaPlateComponent.h b/Engine/Plugins/Media/MediaPlate/Source/MediaPlate/Public/MediaPlateComponent.h
--- a/Engine/Plugins/Media/MediaPlate/Source/MediaPlate/Public/MediaPlateComponent.h
+++ b/Engine/Plugins/Media/MediaPlate/Source/MediaPlate/Public/MediaPlateComponent.h
@@ -34,6 +34,8 @@ enum class EMediaPlateEventState : uint8
 	Reverse,
 	Forward,
 	Rewind,
+	Next,
+	Previous,
 	MAX
 };
 
diff --git a/Engine/Plugins/Media/MediaPlate/Source/MediaPlateEditor/Private/Toolkits/MediaPlateEditorToolkit.cpp b/Engine/Plugins/Media/MediaPlate/Source/MediaPlateEditor/Private/Toolkits/MediaPlateEditorToolkit.cpp
--- a/Engine/Plugins/Media/MediaPlate/Source/MediaPlateEditor/Private/Toolkits/MediaPlateEditorToolkit.cpp
+++ b/Engine/Plugins/Media/MediaPlate/Source/MediaPlateEditor/Private/Toolkits/MediaPlateEditorToolkit.cpp
@@ -6,6 +6,7 @@
 #include "Engine/Engine.h"
 #include "EditorReimportHandler.h"
 #include "Framework/MultiBox/MultiBoxBuilder.h"
+#include "IMediaAssetsModule.h"
 #include "MediaPlate.h"
 #include "MediaPlateComponent.h"
 #include "MediaPlateEditorModule.h"
@@ -262,10 +263,16 @@ void FMediaPlateEditorToolkit::PostRedo(bool bSuccess)
 void FMediaPlateEditorToolkit::BindCommands()
 {
 	const FMediaPlateEditorCommands& Commands = FMediaPlateEditorCommands::Get();
+	IMediaAssetsModule* MediaAssets = FModuleManager::LoadModulePtr<IMediaAssetsModule>("MediaAssets");
 
 	ToolkitCommands->MapAction(
 		Commands.CloseMedia,
-		FExecuteAction::CreateLambda([this] { MediaPlate->GetMediaPlayer()->Close(); }),
+		FExecuteAction::CreateLambda([this, MediaAssets] {
+			MediaPlate->GetMediaPlayer()->Close();
+
+			TArray<FString> ActorsPathNames = { MediaPlate->GetOwner()->GetPathName() };
+			MediaAssets->BroadcastOnMediaStateChangedEvent(ActorsPathNames, (uint8)EMediaPlateEventState::Close);
+		}),
 		FCanExecuteAction::CreateLambda([this] {
 			TObjectPtr<UMediaPlayer> MediaPlayer = MediaPlate->GetMediaPlayer();
 			return (MediaPlayer != nullptr) && !MediaPlayer->GetUrl().IsEmpty();
@@ -274,7 +281,12 @@ void FMediaPlateEditorToolkit::BindCommands()
 
 	ToolkitCommands->MapAction(
 		Commands.ForwardMedia,
-		FExecuteAction::CreateLambda([this]{ MediaPlate->GetMediaPlayer()->SetRate(GetForwardRate()); }),
+		FExecuteAction::CreateLambda([this, MediaAssets] {
+			MediaPlate->GetMediaPlayer()->SetRate(GetForwardRate());
+
+			TArray<FString> ActorsPathNames = { MediaPlate->GetOwner()->GetPathName() };
+			MediaAssets->BroadcastOnMediaStateChangedEvent(ActorsPathNames, (uint8)EMediaPlateEventState::Forward);
+		}),
 		FCanExecuteAction::CreateLambda([this]{
 			TObjectPtr<UMediaPlayer> MediaPlayer = MediaPlate->GetMediaPlayer();
 			return (MediaPlayer != nullptr) && MediaPlayer->IsReady() && MediaPlayer->SupportsRate(GetForwardRate(), false);
@@ -283,7 +295,12 @@ void FMediaPlateEditorToolkit::BindCommands()
 
 	ToolkitCommands->MapAction(
 		Commands.NextMedia,
-		FExecuteAction::CreateLambda([this]{ MediaPlate->Next(); }),
+		FExecuteAction::CreateLambda([this, MediaAssets] {
+			MediaPlate->Next();
+
+			TArray<FString> ActorsPathNames = { MediaPlate->GetOwner()->GetPathName() };
+			MediaAssets->BroadcastOnMediaStateChangedEvent(ActorsPathNames, (uint8)EMediaPlateEventState::Next);
+		}),
 		FCanExecuteAction::CreateLambda([this]{
 			return (MediaPlate->GetMediaPlaylist() != nullptr) &&
 				(MediaPlate->GetMediaPlaylist()->Num() > 1);
@@ -292,13 +309,30 @@ void FMediaPlateEditorToolkit::BindCommands()
 
 	ToolkitCommands->MapAction(
 		Commands.OpenMedia,
-		FExecuteAction::CreateLambda([this] { MediaPlate->Open(); }),
+		FExecuteAction::CreateLambda([this, MediaAssets] {
+			MediaPlate->Open();
+
+			// Tell the editor module that this media plate is playing.
+			FMediaPlateEditorModule* EditorModule = FModuleManager::LoadModulePtr<FMediaPlateEditorModule>("MediaPlateEditor");
+			if (EditorModule != nullptr)
+			{
+				EditorModule->MediaPlateStartedPlayback(MediaPlate);
+			}
+
+			TArray<FString> ActorsPathNames = { MediaPlate->GetOwner()->GetPathName() };
+			MediaAssets->BroadcastOnMediaStateChangedEvent(ActorsPathNames, (uint8)EMediaPlateEventState::Open);
+		}),
 		FCanExecuteAction::CreateLambda([this] { return true; })
 	);
 
 	ToolkitCommands->MapAction(
 		Commands.PauseMedia,
-		FExecuteAction::CreateLambda([this]{ MediaPlate->Pause(); }),
+		FExecuteAction::CreateLambda([this, MediaAssets] {
+			MediaPlate->Pause();
+
+			TArray<FString> ActorsPathNames = { MediaPlate->GetOwner()->GetPathName() };
+			MediaAssets->BroadcastOnMediaStateChangedEvent(ActorsPathNames, (uint8)EMediaPlateEventState::Pause);
+		}),
 		FCanExecuteAction::CreateLambda([this]{
 			TObjectPtr<UMediaPlayer> MediaPlayer = MediaPlate->GetMediaPlayer();
 			return (MediaPlayer != nullptr) &&MediaPlayer->CanPause() && !MediaPlayer->IsPaused();
@@ -307,7 +341,12 @@ void FMediaPlateEditorToolkit::BindCommands()
 
 	ToolkitCommands->MapAction(
 		Commands.PlayMedia,
-		FExecuteAction::CreateLambda([this]{ MediaPlate->GetMediaPlayer()->Play(); }),
+		FExecuteAction::CreateLambda([this, MediaAssets] {
+			MediaPlate->GetMediaPlayer()->Play();
+			
+			TArray<FString> ActorsPathNames = { MediaPlate->GetOwner()->GetPathName() };
+			MediaAssets->BroadcastOnMediaStateChangedEvent(ActorsPathNames, (uint8)EMediaPlateEventState::Play);
+		}),
 		FCanExecuteAction::CreateLambda([this]{
 			TObjectPtr<UMediaPlayer> MediaPlayer = MediaPlate->GetMediaPlayer(); 
 			return (MediaPlayer != nullptr) && MediaPlayer->IsReady() && (!MediaPlayer->IsPlaying() || (MediaPlayer->GetRate() != 1.0f));
@@ -316,7 +355,12 @@ void FMediaPlateEditorToolkit::BindCommands()
 
 	ToolkitCommands->MapAction(
 		Commands.PreviousMedia,
-		FExecuteAction::CreateLambda([this]{ MediaPlate->Previous(); }),
+		FExecuteAction::CreateLambda([this, MediaAssets] {
+			MediaPlate->Previous();
+
+			TArray<FString> ActorsPathNames = { MediaPlate->GetOwner()->GetPathName() };
+			MediaAssets->BroadcastOnMediaStateChangedEvent(ActorsPathNames, (uint8)EMediaPlateEventState::Previous);
+		}),
 		FCanExecuteAction::CreateLambda([this]{
 			return (MediaPlate->GetMediaPlaylist() != nullptr) &&
 					(MediaPlate->GetMediaPlaylist()->Num() > 1);
@@ -325,7 +369,12 @@ void FMediaPlateEditorToolkit::BindCommands()
 
 	ToolkitCommands->MapAction(
 		Commands.ReverseMedia,
-		FExecuteAction::CreateLambda([this]{ MediaPlate->GetMediaPlayer()->SetRate(GetReverseRate()); } ),
+		FExecuteAction::CreateLambda([this, MediaAssets] {
+			MediaPlate->GetMediaPlayer()->SetRate(GetReverseRate());
+
+			TArray<FString> ActorsPathNames = { MediaPlate->GetOwner()->GetPathName() };
+			MediaAssets->BroadcastOnMediaStateChangedEvent(ActorsPathNames, (uint8)EMediaPlateEventState::Reverse);
+		}),
 		FCanExecuteAction::CreateLambda([this]{
 			TObjectPtr<UMediaPlayer> MediaPlayer = MediaPlate->GetMediaPlayer();
 			return (MediaPlayer != nullptr) && MediaPlayer->IsReady() && MediaPlayer->SupportsRate(GetReverseRate(), false);
@@ -334,7 +383,12 @@ void FMediaPlateEditorToolkit::BindCommands()
 
 	ToolkitCommands->MapAction(
 		Commands.RewindMedia,
-		FExecuteAction::CreateLambda([this]{ MediaPlate->GetMediaPlayer()->Rewind(); }),
+		FExecuteAction::CreateLambda([this, MediaAssets] {
+			MediaPlate->GetMediaPlayer()->Rewind();
+
+			TArray<FString> ActorsPathNames = { MediaPlate->GetOwner()->GetPathName() };
+			MediaAssets->BroadcastOnMediaStateChangedEvent(ActorsPathNames, (uint8)EMediaPlateEventState::Rewind);
+		}),
 		FCanExecuteAction::CreateLambda([this]{
 			TObjectPtr<UMediaPlayer> MediaPlayer = MediaPlate->GetMediaPlayer(); 
 			return (MediaPlayer != nullptr) && MediaPlate->GetMediaPlayer()->IsReady() && MediaPlayer->SupportsSeeking() && MediaPlayer->GetTime() > FTimespan::Zero();
