MediaPlate: Only update opaque default on new media plates.

The existing media plates will retain their translucent materials (except for sky and crossfade).

#rnx

#virtualized

diff --git a/Engine/Plugins/Media/MediaPlate/Source/MediaPlate/Private/MediaPlate.cpp b/Engine/Plugins/Media/MediaPlate/Source/MediaPlate/Private/MediaPlate.cpp
--- a/Engine/Plugins/Media/MediaPlate/Source/MediaPlate/Private/MediaPlate.cpp
+++ b/Engine/Plugins/Media/MediaPlate/Source/MediaPlate/Private/MediaPlate.cpp
@@ -156,7 +156,7 @@ UMaterialInterface* AMediaPlate::GetCurrentOverlayMaterial() const
 
 void AMediaPlate::UseDefaultMaterial()
 {
-	UMaterial* DefaultMaterial = LoadObject<UMaterial>(NULL, TEXT("/MediaPlate/M_MediaPlate"), NULL, LOAD_None, NULL);
+	UMaterial* DefaultMaterial = LoadObject<UMaterial>(NULL, TEXT("/MediaPlate/M_MediaPlate_Opaque"), NULL, LOAD_None, NULL);
 	
 	ApplyMaterial(DefaultMaterial);
 
