MediaPlate: Add support for holdout composite.

diff --git a/Engine/Plugins/Media/MediaPlate/MediaPlate.uplugin b/Engine/Plugins/Media/MediaPlate/MediaPlate.uplugin
--- a/Engine/Plugins/Media/MediaPlate/MediaPlate.uplugin
+++ b/Engine/Plugins/Media/MediaPlate/MediaPlate.uplugin
@@ -45,5 +45,9 @@
 			"TargetAllowList": [ "Editor" ],
 			"Enabled": true
 		},
+		{
+			"Name": "HoldoutComposite",
+			"Enabled": true
+		}
 	]
 }
\ No newline at end of file
diff --git a/Engine/Plugins/Media/MediaPlate/Source/MediaPlateEditor/MediaPlateEditor.Build.cs b/Engine/Plugins/Media/MediaPlate/Source/MediaPlateEditor/MediaPlateEditor.Build.cs
--- a/Engine/Plugins/Media/MediaPlate/Source/MediaPlateEditor/MediaPlateEditor.Build.cs
+++ b/Engine/Plugins/Media/MediaPlate/Source/MediaPlateEditor/MediaPlateEditor.Build.cs
@@ -42,7 +42,8 @@ public class MediaPlateEditor : ModuleRules
 				"Slate",
 				"SlateCore",
 				"UnrealEd",
+				"HoldoutComposite"
 			}
-			);
+		);
 	}
 }
diff --git a/Engine/Plugins/Media/MediaPlate/Source/MediaPlateEditor/Private/MediaPlateEditorModule.cpp b/Engine/Plugins/Media/MediaPlate/Source/MediaPlateEditor/Private/MediaPlateEditorModule.cpp
--- a/Engine/Plugins/Media/MediaPlate/Source/MediaPlateEditor/Private/MediaPlateEditorModule.cpp
+++ b/Engine/Plugins/Media/MediaPlate/Source/MediaPlateEditor/Private/MediaPlateEditorModule.cpp
@@ -4,6 +4,7 @@
 
 #include "Editor.h"
 #include "EngineAnalytics.h"
+#include "HoldoutCompositeComponent.h"
 #include "IPlacementModeModule.h"
 #include "ISequencerModule.h"
 #include "LevelEditor.h"
@@ -347,8 +348,9 @@ TSharedRef<FExtender> FMediaPlateEditorModule::ExtendLevelViewportContextMenuFor
 					{
 						const ISlateStyle* Style = &FMediaPlateEditorStyle::Get().Get();
 						check(Style);
-
-						FUIAction Action_ConfigureComposite(
+						
+						// Note: Could be removed now that we have the holdout composite alternative.
+						FUIAction Action_ConfigureOverlayComposite(
 							FExecuteAction::CreateLambda([](AMediaPlate* InMediaPlateActor)
 								{
 									if (!IsValid(InMediaPlateActor))
@@ -370,6 +372,46 @@ TSharedRef<FExtender> FMediaPlateEditorModule::ExtendLevelViewportContextMenuFor
 								,MediaPlateActor)
 						);
 
+						FUIAction Action_ConfigureHoldoutComposite(
+							FExecuteAction::CreateLambda([](AMediaPlate* InMediaPlateActor)
+								{
+									if (!IsValid(InMediaPlateActor))
+									{
+										return;
+									}
+
+									if (!IsValid(InMediaPlateActor->StaticMeshComponent))
+									{
+										UE_LOG(LogMediaPlateEditor, Warning, TEXT("Could not setup holdout composite on media plate, invalid static mesh component."));
+										return;
+									}
+
+									if (InMediaPlateActor->FindComponentByClass<UHoldoutCompositeComponent>() != nullptr)
+									{
+										// Already present, do nothing
+										return;
+									}
+
+									const FScopedTransaction Transaction(LOCTEXT("ApplyHoldoutComposite", "Apply Holdout Composite"));
+									InMediaPlateActor->Modify();
+
+									UHoldoutCompositeComponent* HoldoutComponent = NewObject<UHoldoutCompositeComponent>(InMediaPlateActor
+										, MakeUniqueObjectName(InMediaPlateActor, UHoldoutCompositeComponent::StaticClass())
+										, RF_Transactional);
+
+									InMediaPlateActor->AddInstanceComponent(HoldoutComponent);
+
+									HoldoutComponent->SetupAttachment(InMediaPlateActor->StaticMeshComponent);
+									HoldoutComponent->OnComponentCreated();
+									HoldoutComponent->RegisterComponent();
+
+									// Used to update the sub-object tree-view list
+									FLevelEditorModule& LevelEditor = FModuleManager::LoadModuleChecked<FLevelEditorModule>("LevelEditor");
+									LevelEditor.BroadcastComponentsEdited();
+								}
+								, MediaPlateActor)
+						);
+
 						FUIAction Action_ResetDefault(
 							FExecuteAction::CreateLambda([](AMediaPlate* InMediaPlateActor)
 								{
@@ -393,7 +435,13 @@ TSharedRef<FExtender> FMediaPlateEditorModule::ExtendLevelViewportContextMenuFor
 							LOCTEXT("ApplyOverlayCompositeMats", "Apply Overlay Composite Materials"),
 							LOCTEXT("ApplyOverlayCompositeMats_Tooltip", "Setup the media plate for overlay-compositing to avoid TSR artifacts by replacing relevant materials. This technique is only effective on opaque media currently."),
 							FSlateIcon(Style->GetStyleSetName(), "ClassIcon.MediaPlate"),
-							Action_ConfigureComposite
+							Action_ConfigureOverlayComposite
+						);
+						MenuBuilder.AddMenuEntry(
+							LOCTEXT("AddHoldoutComposite", "Apply Holdout Composite"),
+							LOCTEXT("AddHoldoutComposite_Tooltip", "Setup the media plate for holdout-compositing to avoid TSR artifacts. This technique is only effective on opaque media currently."),
+							FSlateIcon(Style->GetStyleSetName(), "ClassIcon.MediaPlate"),
+							Action_ConfigureHoldoutComposite
 						);
 						MenuBuilder.AddMenuEntry(
 							LOCTEXT("ResetDefaultMats", "Reset Default Materials"),
