MediaPlate: Fixed saving issue with spawnables.

This was previously fixed at CL 22050262, before the playlist object was recently replaced.

#rnx

diff --git a/Engine/Plugins/Media/MediaPlate/Source/MediaPlate/Public/MediaPlateResource.h b/Engine/Plugins/Media/MediaPlate/Source/MediaPlate/Public/MediaPlateResource.h
--- a/Engine/Plugins/Media/MediaPlate/Source/MediaPlate/Public/MediaPlateResource.h
+++ b/Engine/Plugins/Media/MediaPlate/Source/MediaPlate/Public/MediaPlateResource.h
@@ -1,4 +1,4 @@
-﻿// Copyright Epic Games, Inc. All Rights Reserved.
+// Copyright Epic Games, Inc. All Rights Reserved.
 
 #pragma once
 
@@ -110,7 +110,7 @@ friend class UMediaPlateComponent;
 	FString ExternalMediaPath;
 
 	/** Media Source loaded from external path */
-	UPROPERTY()
+	UPROPERTY(Instanced)
 	TObjectPtr<UMediaSource> ExternalMedia;
 
 	/** Media source coming from MediaSource asset*/
@@ -122,6 +122,6 @@ friend class UMediaPlateComponent;
 	TSoftObjectPtr<UMediaPlaylist> SourcePlaylist;
 
 	/** Currently running Playlist asset */
-	UPROPERTY()
+	UPROPERTY(Instanced)
 	TObjectPtr<UMediaPlaylist> ActivePlaylist;
 };
