MediaPlate: Improved composite menu options.

Disallow holdouts when overlays are active, and vice-versa. Updated descriptions, and reordered options.

diff --git a/Engine/Plugins/Media/MediaPlate/Source/MediaPlateEditor/Private/MediaPlateEditorModule.cpp b/Engine/Plugins/Media/MediaPlate/Source/MediaPlateEditor/Private/MediaPlateEditorModule.cpp
--- a/Engine/Plugins/Media/MediaPlate/Source/MediaPlateEditor/Private/MediaPlateEditorModule.cpp
+++ b/Engine/Plugins/Media/MediaPlate/Source/MediaPlateEditor/Private/MediaPlateEditorModule.cpp
@@ -349,29 +349,6 @@ TSharedRef<FExtender> FMediaPlateEditorModule::ExtendLevelViewportContextMenuFor
 					{
 						const ISlateStyle* Style = &FMediaPlateEditorStyle::Get().Get();
 						check(Style);
-						
-						// Note: Could be removed now that we have the holdout composite alternative.
-						FUIAction Action_ConfigureOverlayComposite(
-							FExecuteAction::CreateLambda([](AMediaPlate* InMediaPlateActor)
-								{
-									if (!IsValid(InMediaPlateActor))
-									{
-										return;
-									}
-
-									const FScopedTransaction Transaction(LOCTEXT("ApplyOverlayCompositeMats", "Apply Overlay Composite Materials"));
-									InMediaPlateActor->Modify();
-
-									UMaterial* BasePassMaterial = LoadObject<UMaterial>(NULL, TEXT("/MediaPlate/M_MediaPlate_Masked"), NULL, LOAD_None, NULL);
-									InMediaPlateActor->ApplyMaterial(BasePassMaterial);
-
-									UMaterial* OverlayMaterial = LoadObject<UMaterial>(NULL, TEXT("/MediaPlate/M_MediaPlate_OverlayComp"), NULL, LOAD_None, NULL);
-									InMediaPlateActor->ApplyOverlayMaterial(OverlayMaterial);
-
-									UE::MediaPlate::Private::ApplyTranslucencyScreenPercentageCVar(1);
-								}
-								,MediaPlateActor)
-						);
 
 						FUIAction Action_ConfigureHoldoutComposite(
 							FExecuteAction::CreateLambda([](AMediaPlate* InMediaPlateActor)
@@ -410,6 +387,72 @@ TSharedRef<FExtender> FMediaPlateEditorModule::ExtendLevelViewportContextMenuFor
 									FLevelEditorModule& LevelEditor = FModuleManager::LoadModuleChecked<FLevelEditorModule>("LevelEditor");
 									LevelEditor.BroadcastComponentsEdited();
 								}
+								, MediaPlateActor),
+							FCanExecuteAction::CreateLambda([](AMediaPlate* InMediaPlateActor) -> bool
+								{
+									// Disable when the media plate has an overlay material or an existing holdout
+									return IsValid(InMediaPlateActor) ? !IsValid(InMediaPlateActor->GetCurrentOverlayMaterial()) && !InMediaPlateActor->FindComponentByClass<UHoldoutCompositeComponent>() : false;
+								}
+								, MediaPlateActor)
+						);
+
+						FUIAction Action_RemoveHoldoutComposite(
+							FExecuteAction::CreateLambda([](AMediaPlate* InMediaPlateActor)
+								{
+									if (!IsValid(InMediaPlateActor))
+									{
+										return;
+									}
+
+									UHoldoutCompositeComponent* HoldoutComponent = InMediaPlateActor->FindComponentByClass<UHoldoutCompositeComponent>();
+									if (IsValid(HoldoutComponent))
+									{
+										const FScopedTransaction Transaction(LOCTEXT("RemoveHoldoutComposite", "Remove Holdout Composite"));
+
+										InMediaPlateActor->RemoveInstanceComponent(HoldoutComponent);
+
+										HoldoutComponent->DetachFromComponent(FDetachmentTransformRules::KeepRelativeTransform);
+										HoldoutComponent->UnregisterComponent();
+										HoldoutComponent->DestroyComponent();
+
+										// Used to update the sub-object tree-view list
+										FLevelEditorModule& LevelEditor = FModuleManager::LoadModuleChecked<FLevelEditorModule>("LevelEditor");
+										LevelEditor.BroadcastComponentsEdited();
+									}
+								}
+								, MediaPlateActor),
+							FCanExecuteAction::CreateLambda([](AMediaPlate* InMediaPlateActor) -> bool
+								{
+									return IsValid(InMediaPlateActor) ? InMediaPlateActor->FindComponentByClass<UHoldoutCompositeComponent>() != nullptr : false;
+								}
+								, MediaPlateActor)
+						);
+
+						// Note: Could be removed now that we have the holdout composite alternative.
+						FUIAction Action_ConfigureOverlayComposite(
+							FExecuteAction::CreateLambda([](AMediaPlate* InMediaPlateActor)
+								{
+									if (!IsValid(InMediaPlateActor))
+									{
+										return;
+									}
+
+									const FScopedTransaction Transaction(LOCTEXT("ApplyOverlayCompositeMats", "Apply Overlay Composite Materials"));
+									InMediaPlateActor->Modify();
+
+									UMaterial* BasePassMaterial = LoadObject<UMaterial>(NULL, TEXT("/MediaPlate/M_MediaPlate_Masked"), NULL, LOAD_None, NULL);
+									InMediaPlateActor->ApplyMaterial(BasePassMaterial);
+
+									UMaterial* OverlayMaterial = LoadObject<UMaterial>(NULL, TEXT("/MediaPlate/M_MediaPlate_OverlayComp"), NULL, LOAD_None, NULL);
+									InMediaPlateActor->ApplyOverlayMaterial(OverlayMaterial);
+
+									UE::MediaPlate::Private::ApplyTranslucencyScreenPercentageCVar(1);
+								}
+								, MediaPlateActor),
+							FCanExecuteAction::CreateLambda([](AMediaPlate* InMediaPlateActor) -> bool
+								{
+									return IsValid(InMediaPlateActor) ? !InMediaPlateActor->FindComponentByClass<UHoldoutCompositeComponent>() : false;
+								}
 								, MediaPlateActor)
 						);
 
@@ -433,16 +476,23 @@ TSharedRef<FExtender> FMediaPlateEditorModule::ExtendLevelViewportContextMenuFor
 
 						MenuBuilder.BeginSection("MediaPlate", LOCTEXT("MediaPlateHeading", "Media Plate"));
 						MenuBuilder.AddMenuEntry(
-							LOCTEXT("ApplyOverlayCompositeMats", "Apply Overlay Composite Materials"),
-							LOCTEXT("ApplyOverlayCompositeMats_Tooltip", "Setup the media plate for overlay-compositing to avoid TSR artifacts by replacing relevant materials. This technique is only effective on opaque media currently."),
+							LOCTEXT("AddHoldoutComposite", "Apply Holdout Composite (Preferred)"),
+							LOCTEXT("AddHoldoutComposite_Tooltip", "Setup the media plate for alpha holdout compositing to avoid TSR artifacts and the tone curve. Translucency rendered in front is preserved. This technique is preferred but only effective on opaque media currently."),
 							FSlateIcon(Style->GetStyleSetName(), "ClassIcon.MediaPlate"),
-							Action_ConfigureOverlayComposite
+							Action_ConfigureHoldoutComposite
 						);
 						MenuBuilder.AddMenuEntry(
-							LOCTEXT("AddHoldoutComposite", "Apply Holdout Composite"),
-							LOCTEXT("AddHoldoutComposite_Tooltip", "Setup the media plate for holdout-compositing to avoid TSR artifacts. This technique is only effective on opaque media currently."),
+							LOCTEXT("RemoveHoldoutComposite", "Remove Holdout Composite"),
+							LOCTEXT("RemoveHoldoutComposite_Tooltip", "Remove and destroy the holdout composite component, if present."),
 							FSlateIcon(Style->GetStyleSetName(), "ClassIcon.MediaPlate"),
-							Action_ConfigureHoldoutComposite
+							Action_RemoveHoldoutComposite
+						);
+						MenuBuilder.AddSeparator();
+						MenuBuilder.AddMenuEntry(
+							LOCTEXT("ApplyOverlayCompositeMats", "Apply Overlay Composite Materials"),
+							LOCTEXT("ApplyOverlayCompositeMats_Tooltip", "Setup the media plate for overlay compositing to avoid TSR artifacts by replacing relevant materials. This technique is only effective on opaque media currently."),
+							FSlateIcon(Style->GetStyleSetName(), "ClassIcon.MediaPlate"),
+							Action_ConfigureOverlayComposite
 						);
 						MenuBuilder.AddMenuEntry(
 							LOCTEXT("ResetDefaultMats", "Reset Default Materials"),
@@ -450,6 +500,7 @@ TSharedRef<FExtender> FMediaPlateEditorModule::ExtendLevelViewportContextMenuFor
 							FSlateIcon(Style->GetStyleSetName(), "ClassIcon.MediaPlate"),
 							Action_ResetDefault
 						);
+
 						MenuBuilder.EndSection();
 					},
 					MediaPlateActor
