Fixed constness issues with TSoftObjectPtr.

diff --git a/Engine/Plugins/Animation/ControlRig/Source/ControlRig/Private/Constraints/ControlRigTransformableHandle.cpp b/Engine/Plugins/Animation/ControlRig/Source/ControlRig/Private/Constraints/ControlRigTransformableHandle.cpp
--- a/Engine/Plugins/Animation/ControlRig/Source/ControlRig/Private/Constraints/ControlRigTransformableHandle.cpp
+++ b/Engine/Plugins/Animation/ControlRig/Source/ControlRig/Private/Constraints/ControlRigTransformableHandle.cpp
@@ -639,7 +639,7 @@ void UTransformableControlHandle::ResolveBoundObjects(FMovieSceneSequenceID Loca
 				if (InControlRig->GetObjectBinding() && InControlRig->GetObjectBinding()->GetBoundObject() == Bindable)
 				{
 					UnregisterDelegates();
-					ControlRig = InControlRig;
+					ControlRig = const_cast<UControlRig*>(InControlRig);
 					RegisterDelegates();
 				}
 				break; //just do one
diff --git a/Engine/Plugins/Animation/ControlRig/Source/ControlRig/Private/ControlRig.cpp b/Engine/Plugins/Animation/ControlRig/Source/ControlRig/Private/ControlRig.cpp
--- a/Engine/Plugins/Animation/ControlRig/Source/ControlRig/Private/ControlRig.cpp
+++ b/Engine/Plugins/Animation/ControlRig/Source/ControlRig/Private/ControlRig.cpp
@@ -3531,7 +3531,7 @@ UTransformableControlHandle* UControlRig::CreateTransformableControlHandle(
 	
 	UTransformableControlHandle* CtrlHandle = NewObject<UTransformableControlHandle>(GetTransientPackage(), NAME_None, RF_Transactional);
 	check(CtrlHandle);
-	CtrlHandle->ControlRig = this;
+	CtrlHandle->ControlRig = const_cast<UControlRig*>(this);
 	CtrlHandle->ControlName = InControlName;
 	CtrlHandle->RegisterDelegates();
 	return CtrlHandle;
diff --git a/Engine/Plugins/Animation/ControlRig/Source/ControlRig/Private/ModularRig.cpp b/Engine/Plugins/Animation/ControlRig/Source/ControlRig/Private/ModularRig.cpp
--- a/Engine/Plugins/Animation/ControlRig/Source/ControlRig/Private/ModularRig.cpp
+++ b/Engine/Plugins/Animation/ControlRig/Source/ControlRig/Private/ModularRig.cpp
@@ -16,13 +16,13 @@
 ////////////////////////////////////////////////////////////////////////////////
 
 FModuleInstanceHandle::FModuleInstanceHandle(const UModularRig* InModularRig, const FString& InPath)
-: ModularRig(InModularRig)
+: ModularRig(const_cast<UModularRig*>(InModularRig))
 , Path(InPath)
 {
 }
 
 FModuleInstanceHandle::FModuleInstanceHandle(const UModularRig* InModularRig, const FRigModuleInstance* InModule)
-: ModularRig(InModularRig)
+: ModularRig(const_cast<UModularRig*>(InModularRig))
 , Path(InModule->GetPath())
 {
 }
diff --git a/Engine/Plugins/Animation/DeformerGraph/Source/OptimusCore/Private/OptimusDeformer.cpp b/Engine/Plugins/Animation/DeformerGraph/Source/OptimusCore/Private/OptimusDeformer.cpp
--- a/Engine/Plugins/Animation/DeformerGraph/Source/OptimusCore/Private/OptimusDeformer.cpp
+++ b/Engine/Plugins/Animation/DeformerGraph/Source/OptimusCore/Private/OptimusDeformer.cpp
@@ -2160,7 +2160,7 @@ TArray<FOptimusComputeGraphInfo> UOptimusDeformer::CompileNodeGraphToComputeGrap
 			FOptimusGraphVariableDescription& ValueNodeDescription = ValueNodeDescriptions.AddDefaulted_GetRef();
 			ValueNodeDescription.Name = Optimus::MakeUniqueValueName(ValueProvider->GetValueName(), ValueNodeIndex);
 			ValueNodeDescription.ValueType = ValueProvider->GetValueType()->ShaderValueType;
-			ValueNodeDescription.SourceObject = ValueNode;
+			ValueNodeDescription.SourceObject = const_cast<UOptimusNode*>(ValueNode);
 
 			if (UOptimusNode_ConstantValue const* ConstantNode = Cast<const UOptimusNode_ConstantValue>(ValueNode))
 			{
diff --git a/Engine/Plugins/Animation/DeformerGraph/Source/OptimusCore/Private/OptimusFunctionNodeGraph.cpp b/Engine/Plugins/Animation/DeformerGraph/Source/OptimusCore/Private/OptimusFunctionNodeGraph.cpp
--- a/Engine/Plugins/Animation/DeformerGraph/Source/OptimusCore/Private/OptimusFunctionNodeGraph.cpp
+++ b/Engine/Plugins/Animation/DeformerGraph/Source/OptimusCore/Private/OptimusFunctionNodeGraph.cpp
@@ -27,7 +27,7 @@ FOptimusFunctionNodeGraphHeader UOptimusFunctionNodeGraph::GetHeader() const
 {
 	FOptimusFunctionNodeGraphHeader Header;
 
-	Header.GraphPath = this; 
+	Header.GraphPath = const_cast<UOptimusFunctionNodeGraph*>(this);
 	Header.FunctionName = GetFName();
 	Header.Category = Category;
 
diff --git a/Engine/Plugins/Enterprise/VariantManager/Source/VariantManager/Private/SDependencyRow.cpp b/Engine/Plugins/Enterprise/VariantManager/Source/VariantManager/Private/SDependencyRow.cpp
--- a/Engine/Plugins/Enterprise/VariantManager/Source/VariantManager/Private/SDependencyRow.cpp
+++ b/Engine/Plugins/Enterprise/VariantManager/Source/VariantManager/Private/SDependencyRow.cpp
@@ -206,7 +206,7 @@ void SDependencyRow::OnSelectedVariantSetChanged(TSharedPtr<FText> NewItem, ESel
 			// to be pickable as a dependency in the first place.
 			// These two facts together prevent us from getting to some invalid states when we could e.g. leave the
 			// Variant part of the dependency as None, and then have another variant depend on this one.
-			for ( const UVariant* Variant : VariantSet->GetVariants() )
+			for ( UVariant* Variant : VariantSet->GetVariants() )
 			{
 				if ( ParentVariant->IsValidDependency( Variant ) )
 				{
diff --git a/Engine/Plugins/Experimental/LandscapePatch/Source/LandscapePatch/Private/LandscapePatchManager.cpp b/Engine/Plugins/Experimental/LandscapePatch/Source/LandscapePatch/Private/LandscapePatchManager.cpp
--- a/Engine/Plugins/Experimental/LandscapePatch/Source/LandscapePatch/Private/LandscapePatchManager.cpp
+++ b/Engine/Plugins/Experimental/LandscapePatch/Source/LandscapePatch/Private/LandscapePatchManager.cpp
@@ -365,7 +365,7 @@ bool ALandscapePatchManager::RemovePatch(ULandscapePatchComponent* Patch)
 
 int32 ALandscapePatchManager::GetIndexOfPatch(const ULandscapePatchComponent* Patch) const
 {
-	if (const int32* Index = PatchToIndex.Find(TSoftObjectPtr<ULandscapePatchComponent>(Patch)))
+	if (const int32* Index = PatchToIndex.Find(TSoftObjectPtr<ULandscapePatchComponent>(const_cast<ULandscapePatchComponent*>(Patch))))
 	{
 		return *Index;
 	}
diff --git a/Engine/Plugins/Experimental/Mutable/Source/CustomizableObjectEditor/Private/MuCOE/CustomizableObjectCompiler.cpp b/Engine/Plugins/Experimental/Mutable/Source/CustomizableObjectEditor/Private/MuCOE/CustomizableObjectCompiler.cpp
--- a/Engine/Plugins/Experimental/Mutable/Source/CustomizableObjectEditor/Private/MuCOE/CustomizableObjectCompiler.cpp
+++ b/Engine/Plugins/Experimental/Mutable/Source/CustomizableObjectEditor/Private/MuCOE/CustomizableObjectCompiler.cpp
@@ -932,13 +932,13 @@ void FCustomizableObjectCompiler::CompileInternal(bool bAsync)
 		ModelResources.Skeletons.Reserve(GenerationContext.ReferencedSkeletons.Num());
 		for (const USkeleton* Skeleton : GenerationContext.ReferencedSkeletons)
 		{
-			ModelResources.Skeletons.Emplace(Skeleton);
+			ModelResources.Skeletons.Emplace(const_cast<USkeleton*>(Skeleton));
 		}
 		
 		ModelResources.Materials.Reserve(GenerationContext.ReferencedMaterials.Num());
 		for (const UMaterialInterface* Material : GenerationContext.ReferencedMaterials)
 		{
-			ModelResources.Materials.Emplace(Material);
+			ModelResources.Materials.Emplace(const_cast<UMaterialInterface*>(Material));
 		}
 
 		for (const TPair<TSoftObjectPtr<USkeletalMesh>, FMutableGraphGenerationContext::FGeneratedReferencedMesh>& Pair : GenerationContext.PassthroughMeshMap)
diff --git a/Engine/Plugins/Media/MediaPlate/Source/MediaPlate/Private/MediaPlateResource.cpp b/Engine/Plugins/Media/MediaPlate/Source/MediaPlate/Private/MediaPlateResource.cpp
--- a/Engine/Plugins/Media/MediaPlate/Source/MediaPlate/Private/MediaPlateResource.cpp
+++ b/Engine/Plugins/Media/MediaPlate/Source/MediaPlate/Private/MediaPlateResource.cpp
@@ -1,4 +1,4 @@
-﻿// Copyright Epic Games, Inc. All Rights Reserved.
+// Copyright Epic Games, Inc. All Rights Reserved.
 
 #include "MediaPlateResource.h"
 #include "MediaPlaylist.h"
@@ -17,7 +17,7 @@ void FMediaPlateResource::SetResourceType(EMediaPlateResourceType InType)
 void FMediaPlateResource::SelectAsset(const UMediaSource* InMediaSource, UObject* InOuter)
 {
 	SetResourceType(EMediaPlateResourceType::Asset);
-	MediaAsset = InMediaSource;
+	MediaAsset = const_cast<UMediaSource*>(InMediaSource);
 	RefreshActivePlaylist(InOuter);
 }
 
@@ -32,7 +32,7 @@ void FMediaPlateResource::LoadExternalMedia(const FString& InFilePath, UObject*
 void FMediaPlateResource::SelectPlaylist(const UMediaPlaylist* InPlaylist)
 {
 	SetResourceType(EMediaPlateResourceType::Playlist);
-	SourcePlaylist = InPlaylist;
+	SourcePlaylist = const_cast<UMediaPlaylist*>(InPlaylist);
 	ActivePlaylist = GetSourcePlaylist();
 }
 
@@ -53,12 +53,12 @@ void FMediaPlateResource::Init(const FMediaPlateResource& InOther)
 		ExternalMediaPath = FString(InOther.GetExternalMediaPath());
 	}
 
-	if (const UMediaSource* OtherMediaAsset = InOther.GetMediaAsset())
+	if (UMediaSource* OtherMediaAsset = InOther.GetMediaAsset())
 	{
 		MediaAsset = OtherMediaAsset;
 	}
 
-	if (const UMediaPlaylist* OtherMediaPlaylist = InOther.GetSourcePlaylist())
+	if (UMediaPlaylist* OtherMediaPlaylist = InOther.GetSourcePlaylist())
 	{
 		SourcePlaylist = OtherMediaPlaylist;
 	}
diff --git a/Engine/Plugins/MovieScene/MovieRenderPipeline/Source/MovieRenderPipelineCore/Private/Graph/MovieGraphRenderLayerSubsystem.cpp b/Engine/Plugins/MovieScene/MovieRenderPipeline/Source/MovieRenderPipelineCore/Private/Graph/MovieGraphRenderLayerSubsystem.cpp
--- a/Engine/Plugins/MovieScene/MovieRenderPipeline/Source/MovieRenderPipelineCore/Private/Graph/MovieGraphRenderLayerSubsystem.cpp
+++ b/Engine/Plugins/MovieScene/MovieRenderPipeline/Source/MovieRenderPipelineCore/Private/Graph/MovieGraphRenderLayerSubsystem.cpp
@@ -209,7 +209,7 @@ void UMovieGraphRenderPropertyModifier::ApplyModifier(const UWorld* World)
 		const TSet<AActor*> MatchingActors = Collection->Evaluate(World);
 		ModifiedActors.Reserve(MatchingActors.Num());
 
-		for (const AActor* Actor : MatchingActors)
+		for (AActor* Actor : MatchingActors)
 		{
 			NewVisibilityState.Actor = Actor;
 			
@@ -225,7 +225,7 @@ void UMovieGraphRenderPropertyModifier::ApplyModifier(const UWorld* World)
 			OriginalVisibilityState.Components.Reserve(Components.Num());
 			NewVisibilityState.Components.Empty(Components.Num());
 			
-			for (const USceneComponent* SceneComponent : Components)
+			for (USceneComponent* SceneComponent : Components)
 			{
 #if WITH_EDITORONLY_DATA
 				// Don't bother processing editor-only components (editor billboard icons, text, etc)
@@ -2284,7 +2284,7 @@ void UMovieGraphConditionGroupQuery_DataLayer::AddDataLayers(const TArray<const
 		if (DataLayer && !DataLayers.Contains(DataLayer))
 		{
 			DataLayers.Add(const_cast<UDataLayerAsset*>(DataLayer));
-			ListDataSource.Add(MakeShared<TSoftObjectPtr<UDataLayerAsset>>(DataLayer));
+			ListDataSource.Add(MakeShared<TSoftObjectPtr<UDataLayerAsset>>(const_cast<UDataLayerAsset*>(DataLayer)));
 		}
 	}
 	
diff --git a/Engine/Plugins/MovieScene/MovieRenderPipeline/Source/MovieRenderPipelineCore/Private/MoviePipelineQueue.cpp b/Engine/Plugins/MovieScene/MovieRenderPipeline/Source/MovieRenderPipelineCore/Private/MoviePipelineQueue.cpp
--- a/Engine/Plugins/MovieScene/MovieRenderPipeline/Source/MovieRenderPipelineCore/Private/MoviePipelineQueue.cpp
+++ b/Engine/Plugins/MovieScene/MovieRenderPipeline/Source/MovieRenderPipelineCore/Private/MoviePipelineQueue.cpp
@@ -316,7 +316,7 @@ void UMoviePipelineExecutorJob::SetGraphPreset(const UMovieGraphConfig* InGraphP
 	Modify();
 #endif
 	
-	GraphPreset = InGraphPreset;
+	GraphPreset = const_cast<UMovieGraphConfig*>(InGraphPreset);
 
 	// If the graph is being cleared out, also clear out all graphs on the job's shots. A shot cannot have a graph assigned to it while
 	// the parent job is using a legacy config.
@@ -411,7 +411,7 @@ void UMoviePipelineExecutorShot::SetGraphPreset(const UMovieGraphConfig* InGraph
 	Modify();
 #endif
 	
-	GraphPreset = InGraphPreset;
+	GraphPreset = const_cast<UMovieGraphConfig*>(InGraphPreset);
 
 	if (bUpdateVariableAssignments)
 	{
diff --git a/Engine/Plugins/PCG/Source/PCG/Private/Grid/PCGGridDescriptor.cpp b/Engine/Plugins/PCG/Source/PCG/Private/Grid/PCGGridDescriptor.cpp
--- a/Engine/Plugins/PCG/Source/PCG/Private/Grid/PCGGridDescriptor.cpp
+++ b/Engine/Plugins/PCG/Source/PCG/Private/Grid/PCGGridDescriptor.cpp
@@ -75,7 +75,7 @@ FPCGGridDescriptor& FPCGGridDescriptor::SetDataLayerAssets(const TArray<TSoftObj
 FPCGGridDescriptor& FPCGGridDescriptor::SetDataLayerAssets(const TArray<const UDataLayerAsset*>& InDataLayerAssets)
 {
 	TArray<TSoftObjectPtr<UDataLayerAsset>> DataLayerAssetPtrs;
-	Algo::Transform(InDataLayerAssets, DataLayerAssetPtrs, [](const UDataLayerAsset* DataLayerAsset) { return TSoftObjectPtr<UDataLayerAsset>(DataLayerAsset); });
+	Algo::Transform(InDataLayerAssets, DataLayerAssetPtrs, [](const UDataLayerAsset* DataLayerAsset) { return TSoftObjectPtr<UDataLayerAsset>(const_cast<UDataLayerAsset*>(DataLayerAsset)); });
 	return SetDataLayerAssets(DataLayerAssetPtrs);
 }
 
diff --git a/Engine/Plugins/Runtime/DataRegistry/Source/DataRegistry/Private/SoftDataRegistryOrTable.cpp b/Engine/Plugins/Runtime/DataRegistry/Source/DataRegistry/Private/SoftDataRegistryOrTable.cpp
--- a/Engine/Plugins/Runtime/DataRegistry/Source/DataRegistry/Private/SoftDataRegistryOrTable.cpp
+++ b/Engine/Plugins/Runtime/DataRegistry/Source/DataRegistry/Private/SoftDataRegistryOrTable.cpp
@@ -24,7 +24,7 @@ FSoftDataRegistryOrTable::FSoftDataRegistryOrTable(const UDataTable* InDataTable
 	}
 	else
 	{ 
-		Table = InDataTable;
+		Table = const_cast<UDataTable*>(InDataTable);
 		bUseDataRegistry = false;
 	}
 }
diff --git a/Engine/Source/Editor/AnimationSettings/Private/AnimationSettingsModule.cpp b/Engine/Source/Editor/AnimationSettings/Private/AnimationSettingsModule.cpp
--- a/Engine/Source/Editor/AnimationSettings/Private/AnimationSettingsModule.cpp
+++ b/Engine/Source/Editor/AnimationSettings/Private/AnimationSettingsModule.cpp
@@ -19,7 +19,7 @@ DEFINE_LOG_CATEGORY(LogAnimationSettings);
 
 void FAnimationSettingsModule::OnAssetRenamed(const FAssetData& InAssetData, const FString& InOldName)
 {
-	if (const UUserDefinedStruct* UserDefinedStruct = Cast<UUserDefinedStruct>(InAssetData.GetAsset()))
+	if (UUserDefinedStruct* UserDefinedStruct = Cast<UUserDefinedStruct>(InAssetData.GetAsset()))
 	{
 		if (UAnimationSettings* Settings = UAnimationSettings::Get())
 		{
diff --git a/Engine/Source/Runtime/Engine/Private/HLODProxy.cpp b/Engine/Source/Runtime/Engine/Private/HLODProxy.cpp
--- a/Engine/Source/Runtime/Engine/Private/HLODProxy.cpp
+++ b/Engine/Source/Runtime/Engine/Private/HLODProxy.cpp
@@ -35,7 +35,7 @@
 
 void UHLODProxy::SetMap(const UWorld* InMap)
 {
-	OwningMap = InMap;
+	OwningMap = const_cast<UWorld*>(InMap);
 }
 
 TSoftObjectPtr<UWorld> UHLODProxy::GetMap() const
diff --git a/Engine/Source/Runtime/Engine/Private/Level.cpp b/Engine/Source/Runtime/Engine/Private/Level.cpp
--- a/Engine/Source/Runtime/Engine/Private/Level.cpp
+++ b/Engine/Source/Runtime/Engine/Private/Level.cpp
@@ -3123,7 +3123,7 @@ void ULevel::OnLevelLoaded()
 	// Set level's associated WorldPartitionRuntimeCell for dynamically injected cells
 	if (LevelStreaming && !WorldPartitionRuntimeCell.GetUniqueID().IsValid())
 	{
-		WorldPartitionRuntimeCell = Cast<const UWorldPartitionRuntimeCell>(LevelStreaming->GetWorldPartitionCell());
+		WorldPartitionRuntimeCell = Cast<UWorldPartitionRuntimeCell>(const_cast<IWorldPartitionCell*>(LevelStreaming->GetWorldPartitionCell()));
 	}
 
 	if (ULevel* OwningLevel = ULevelInstanceSubsystem::GetOwningLevel(this, true))
diff --git a/Engine/Source/Runtime/Engine/Private/LevelInstance/LevelInstancePropertyOverrideAsset.cpp b/Engine/Source/Runtime/Engine/Private/LevelInstance/LevelInstancePropertyOverrideAsset.cpp
--- a/Engine/Source/Runtime/Engine/Private/LevelInstance/LevelInstancePropertyOverrideAsset.cpp
+++ b/Engine/Source/Runtime/Engine/Private/LevelInstance/LevelInstancePropertyOverrideAsset.cpp
@@ -118,7 +118,7 @@ TSoftObjectPtr<ULevelInstancePropertyOverrideAsset> ULevelInstancePropertyOverri
 		}
 	}
 
-	return TSoftObjectPtr<ULevelInstancePropertyOverrideAsset>(this);
+	return TSoftObjectPtr<ULevelInstancePropertyOverrideAsset>(const_cast<ULevelInstancePropertyOverrideAsset*>(this));
 }
 
 FActorContainerPath ULevelInstancePropertyOverrideAsset::GetContainerPropertyOverridePath(ILevelInstanceInterface* InParent, ILevelInstanceInterface* InChild)
diff --git a/Engine/Source/Runtime/Engine/Private/WorldPartition/HLOD/HLODSourceActorsFromLevel.cpp b/Engine/Source/Runtime/Engine/Private/WorldPartition/HLOD/HLODSourceActorsFromLevel.cpp
--- a/Engine/Source/Runtime/Engine/Private/WorldPartition/HLOD/HLODSourceActorsFromLevel.cpp
+++ b/Engine/Source/Runtime/Engine/Private/WorldPartition/HLOD/HLODSourceActorsFromLevel.cpp
@@ -46,7 +46,7 @@ uint32 UWorldPartitionHLODSourceActorsFromLevel::GetHLODHash() const
 
 void UWorldPartitionHLODSourceActorsFromLevel::SetSourceLevel(const UWorld* InSourceLevel)
 {
-	SourceLevel = InSourceLevel;
+	SourceLevel = const_cast<UWorld*>(InSourceLevel);
 }
 
 const TSoftObjectPtr<UWorld>& UWorldPartitionHLODSourceActorsFromLevel::GetSourceLevel() const
diff --git a/Engine/Source/Runtime/Engine/Private/WorldPartition/WorldPartitionLevelHelper.cpp b/Engine/Source/Runtime/Engine/Private/WorldPartition/WorldPartitionLevelHelper.cpp
--- a/Engine/Source/Runtime/Engine/Private/WorldPartition/WorldPartitionLevelHelper.cpp
+++ b/Engine/Source/Runtime/Engine/Private/WorldPartition/WorldPartitionLevelHelper.cpp
@@ -421,7 +421,7 @@ ULevel* FWorldPartitionLevelHelper::CreateEmptyLevelForRuntimeCell(const UWorldP
 	check(NewLevel->Model);
 	check(!NewLevel->bIsVisible);
 
-	NewLevel->WorldPartitionRuntimeCell = Cell;
+	NewLevel->WorldPartitionRuntimeCell = const_cast<UWorldPartitionRuntimeCell*>(Cell);
 	
 	// Mark the level package as fully loaded
 	CellPackage->MarkAsFullyLoaded();
