[Media Plate Editor] Adding Outliner Context Menu action to copy the media texture reference to clipboard.

This allows to paste the media texture in Material Designer's texture slot (or other texture sample in material).

diff --git a/Engine/Plugins/Media/MediaPlate/Source/MediaPlateEditor/MediaPlateEditor.Build.cs b/Engine/Plugins/Media/MediaPlate/Source/MediaPlateEditor/MediaPlateEditor.Build.cs
--- a/Engine/Plugins/Media/MediaPlate/Source/MediaPlateEditor/MediaPlateEditor.Build.cs
+++ b/Engine/Plugins/Media/MediaPlate/Source/MediaPlateEditor/MediaPlateEditor.Build.cs
@@ -19,6 +19,7 @@ public class MediaPlateEditor : ModuleRules
 			new string[]
 			{
 				"AssetDefinition",
+				"ApplicationCore",
 				"CinematicCamera",
 				"CoreUObject",
 				"DesktopWidgets",
diff --git a/Engine/Plugins/Media/MediaPlate/Source/MediaPlateEditor/Private/MediaPlateEditorModule.cpp b/Engine/Plugins/Media/MediaPlate/Source/MediaPlateEditor/Private/MediaPlateEditorModule.cpp
--- a/Engine/Plugins/Media/MediaPlate/Source/MediaPlateEditor/Private/MediaPlateEditorModule.cpp
+++ b/Engine/Plugins/Media/MediaPlate/Source/MediaPlateEditor/Private/MediaPlateEditorModule.cpp
@@ -4,6 +4,7 @@
 
 #include "Editor.h"
 #include "EngineAnalytics.h"
+#include "HAL/PlatformApplicationMisc.h"
 #include "HoldoutCompositeComponent.h"
 #include "IPlacementModeModule.h"
 #include "ISequencerModule.h"
@@ -17,6 +18,7 @@
 #include "MediaPlateResourceCustomization.h"
 #include "MediaSoundComponent.h"
 #include "MediaSource.h"
+#include "MediaTexture.h"
 #include "Models/MediaPlateEditorCommands.h"
 #include "Modules/ModuleManager.h"
 #include "PropertyEditorModule.h"
@@ -474,6 +476,25 @@ TSharedRef<FExtender> FMediaPlateEditorModule::ExtendLevelViewportContextMenuFor
 								, MediaPlateActor)
 						);
 
+						FUIAction Action_CopyMediaTextureReference(
+							FExecuteAction::CreateLambda([](AMediaPlate* InMediaPlateActor)
+								{
+									if (!IsValid(InMediaPlateActor))
+									{
+										return;
+									}
+
+									if (UMediaPlateComponent* MediaPlateComponent = InMediaPlateActor->MediaPlateComponent)
+									{
+										if (const UMediaTexture* MediaTexture = MediaPlateComponent->GetMediaTexture(0))
+										{
+											FPlatformApplicationMisc::ClipboardCopy(*FAssetData(MediaTexture).GetExportTextName());
+										}
+									}
+								}
+								, MediaPlateActor)
+						);
+
 						MenuBuilder.BeginSection("MediaPlate", LOCTEXT("MediaPlateHeading", "Media Plate"));
 						MenuBuilder.AddMenuEntry(
 							LOCTEXT("AddHoldoutComposite", "Apply Holdout Composite (Preferred)"),
@@ -501,6 +522,14 @@ TSharedRef<FExtender> FMediaPlateEditorModule::ExtendLevelViewportContextMenuFor
 							Action_ResetDefault
 						);
 
+						MenuBuilder.AddSeparator();
+						MenuBuilder.AddMenuEntry(
+							LOCTEXT("CopyMediaTextureReference", "Copy Media Texture Reference"),
+							LOCTEXT("CopyMediaTextureReference_Tooltip", "Copy Embedded Media Texture Reference to Clipboard."),
+							FSlateIcon(Style->GetStyleSetName(), "ClassIcon.MediaPlate"),
+							Action_CopyMediaTextureReference
+						);
+
 						MenuBuilder.EndSection();
 					},
 					MediaPlateActor
