[Media Plate] Auto play logic is applied if the media source is modified while playing (by [REDACTED_PERSON] for instance).

The feature was originally added in CL 32313386.

diff --git a/Engine/Plugins/Media/MediaPlate/Source/MediaPlate/Private/MediaPlateComponent.cpp b/Engine/Plugins/Media/MediaPlate/Source/MediaPlate/Private/MediaPlateComponent.cpp
--- a/Engine/Plugins/Media/MediaPlate/Source/MediaPlate/Private/MediaPlateComponent.cpp
+++ b/Engine/Plugins/Media/MediaPlate/Source/MediaPlate/Private/MediaPlateComponent.cpp
@@ -915,7 +915,7 @@ float UMediaPlateComponent::GetReverseRate(UMediaPlayer* MediaPlayer)
 	return 2.0f * Rate;
 }
 
-void UMediaPlateComponent::RestartPlayer()
+bool UMediaPlateComponent::RestartPlayer()
 {
 	if (MediaPlayer != nullptr)
 	{
@@ -923,8 +923,10 @@ void UMediaPlateComponent::RestartPlayer()
 		{
 			MediaPlayer->Close();
 			Open();
+			return true;
 		}
 	}
+	return false;
 }
 
 void UMediaPlateComponent::StopClockSink()
@@ -1356,6 +1358,11 @@ void UMediaPlateComponent::RefreshMediaPlateResource()
 {
 	MediaPlateResource.RefreshActivePlaylist(this);
 
+	if (RestartPlayer())
+	{
+		return;
+	}
+
 	constexpr EMediaPlateEventState State = EMediaPlateEventState::Close;
 	SwitchStates(State);
 
diff --git a/Engine/Plugins/Media/MediaPlate/Source/MediaPlate/Public/MediaPlateComponent.h b/Engine/Plugins/Media/MediaPlate/Source/MediaPlate/Public/MediaPlateComponent.h
--- a/Engine/Plugins/Media/MediaPlate/Source/MediaPlate/Public/MediaPlateComponent.h
+++ b/Engine/Plugins/Media/MediaPlate/Source/MediaPlate/Public/MediaPlateComponent.h
@@ -340,7 +340,7 @@ private:
 	 */
 	void PlayOnlyWhenVisibleChanged();
 
-	void RestartPlayer();
+	bool RestartPlayer();
 
 	/**
 	 * If true, then we want the media plate to play.
