MediaPlate: Fix issue where existing material instances would be converted into a material instance constant, preventing further parameter modifications.

diff --git a/Engine/Plugins/Media/MediaPlate/Source/MediaPlate/Private/MediaPlate.cpp b/Engine/Plugins/Media/MediaPlate/Source/MediaPlate/Private/MediaPlate.cpp
--- a/Engine/Plugins/Media/MediaPlate/Source/MediaPlate/Private/MediaPlate.cpp
+++ b/Engine/Plugins/Media/MediaPlate/Source/MediaPlate/Private/MediaPlate.cpp
@@ -224,7 +224,6 @@ void AMediaPlate::ApplyMaterial(UMaterialInterface* Material)
 		}
 		else
 		{
-			UMaterialInstanceDynamic* MID = Cast<UMaterialInstanceDynamic>(Material);
 			UMaterialInterface* Result = nullptr;
 
 			// See if we can modify this material.
@@ -240,8 +239,14 @@ void AMediaPlate::ApplyMaterial(UMaterialInterface* Material)
 				MediaPlateComponent->SetNumberOfTextures(1);
 				LastMaterial = Material;
 			}
-			else if (MID != nullptr)
+			else if (UMaterialInstanceDynamic* InMID = Cast<UMaterialInstanceDynamic>(Material))
 			{
+				SetMIDParameters(InMID);
+				Result = InMID;
+			}
+			else if (Material->IsA(UMaterialInstance::StaticClass()))
+			{
+				UMaterialInstanceDynamic* MID = UMaterialInstanceDynamic::Create(Material, StaticMeshComponent);
 				SetMIDParameters(MID);
 				Result = MID;
 			}
@@ -289,6 +294,12 @@ void AMediaPlate::ApplyOverlayMaterial(UMaterialInterface* InOverlayMaterial)
 				SetMIDParameters(MID);
 				Result = MID;
 			}
+			else if (InOverlayMaterial->IsA(UMaterialInstance::StaticClass()))
+			{
+				UMaterialInstanceDynamic* OverlayMID = UMaterialInstanceDynamic::Create(InOverlayMaterial, StaticMeshComponent);
+				SetMIDParameters(OverlayMID);
+				Result = OverlayMID;
+			}
 			else
 			{
 				MediaPlateComponent->SetNumberOfTextures(1);
