MediaPlate: Mark overlay composite as deprecated, now superceded by holdout composite.

While it requires a project setting activation, holdout supports translucency in front, bypasses tonecurve color transformations, is not dependent on r.Translucency.ScreenPercentage.Basis and should be more robust overall.

diff --git a/Engine/Plugins/Media/MediaPlate/Source/MediaPlateEditor/Private/MediaPlateEditorModule.cpp b/Engine/Plugins/Media/MediaPlate/Source/MediaPlateEditor/Private/MediaPlateEditorModule.cpp
--- a/Engine/Plugins/Media/MediaPlate/Source/MediaPlateEditor/Private/MediaPlateEditorModule.cpp
+++ b/Engine/Plugins/Media/MediaPlate/Source/MediaPlateEditor/Private/MediaPlateEditorModule.cpp
@@ -497,7 +497,7 @@ TSharedRef<FExtender> FMediaPlateEditorModule::ExtendLevelViewportContextMenuFor
 
 						MenuBuilder.BeginSection("MediaPlate", LOCTEXT("MediaPlateHeading", "Media Plate"));
 						MenuBuilder.AddMenuEntry(
-							LOCTEXT("AddHoldoutComposite", "Apply Holdout Composite (Preferred)"),
+							LOCTEXT("AddHoldoutComposite", "Apply Holdout Composite"),
 							LOCTEXT("AddHoldoutComposite_Tooltip", "Setup the media plate for alpha holdout compositing to avoid TSR artifacts and the tone curve. Translucency rendered in front is preserved. This technique is preferred but only effective on opaque media currently."),
 							FSlateIcon(Style->GetStyleSetName(), "ClassIcon.MediaPlate"),
 							Action_ConfigureHoldoutComposite
@@ -510,7 +510,7 @@ TSharedRef<FExtender> FMediaPlateEditorModule::ExtendLevelViewportContextMenuFor
 						);
 						MenuBuilder.AddSeparator();
 						MenuBuilder.AddMenuEntry(
-							LOCTEXT("ApplyOverlayCompositeMats", "Apply Overlay Composite Materials"),
+							LOCTEXT("ApplyOverlayCompositeMats", "Apply Overlay Composite Materials (Deprecated)"),
 							LOCTEXT("ApplyOverlayCompositeMats_Tooltip", "Setup the media plate for overlay compositing to avoid TSR artifacts by replacing relevant materials. This technique is only effective on opaque media currently."),
 							FSlateIcon(Style->GetStyleSetName(), "ClassIcon.MediaPlate"),
 							Action_ConfigureOverlayComposite
