MediaPlate: Fixed localization warning caused by previous CL.

#rnx

diff --git a/Engine/Plugins/Media/MediaPlate/Source/MediaPlateEditor/Private/MediaPlateEditorModule.cpp b/Engine/Plugins/Media/MediaPlate/Source/MediaPlateEditor/Private/MediaPlateEditorModule.cpp
--- a/Engine/Plugins/Media/MediaPlate/Source/MediaPlateEditor/Private/MediaPlateEditorModule.cpp
+++ b/Engine/Plugins/Media/MediaPlate/Source/MediaPlateEditor/Private/MediaPlateEditorModule.cpp
@@ -439,7 +439,7 @@ TSharedRef<FExtender> FMediaPlateEditorModule::ExtendLevelViewportContextMenuFor
 										return;
 									}
 
-									const FScopedTransaction Transaction(LOCTEXT("ApplyOverlayCompositeMats", "Apply Overlay Composite Materials"));
+									const FScopedTransaction Transaction(LOCTEXT("ApplyOverlayCompositeMats", "Apply Overlay Composite Materials (Deprecated)"));
 									InMediaPlateActor->Modify();
 
 									UMaterial* BasePassMaterial = LoadObject<UMaterial>(NULL, TEXT("/MediaPlate/M_MediaPlate_Masked"), NULL, LOAD_None, NULL);
