Holdout: Fixed alpha threshold, and removed comments claiming that translucency is not supported.

diff --git a/Engine/Plugins/Experimental/Compositing/HoldoutComposite/Shaders/Private/HoldoutCompositeDilate.usf b/Engine/Plugins/Experimental/Compositing/HoldoutComposite/Shaders/Private/HoldoutCompositeDilate.usf
--- a/Engine/Plugins/Experimental/Compositing/HoldoutComposite/Shaders/Private/HoldoutCompositeDilate.usf
+++ b/Engine/Plugins/Experimental/Compositing/HoldoutComposite/Shaders/Private/HoldoutCompositeDilate.usf
@@ -7,7 +7,7 @@
 #endif
 
 #define KERNEL_COUNT 8
-#define ALPHA_THRESHOLD 0.5f // Custom render passes are opaque-only currently
+#define ALPHA_THRESHOLD (1.0f - 1e-9f)
 
 // Number of texels arround the group tile for scene color.
 #define LDS_TILE_BORDER_SIZE (1)
diff --git a/Engine/Plugins/Media/MediaPlate/Source/MediaPlateEditor/Private/MediaPlateEditorModule.cpp b/Engine/Plugins/Media/MediaPlate/Source/MediaPlateEditor/Private/MediaPlateEditorModule.cpp
--- a/Engine/Plugins/Media/MediaPlate/Source/MediaPlateEditor/Private/MediaPlateEditorModule.cpp
+++ b/Engine/Plugins/Media/MediaPlate/Source/MediaPlateEditor/Private/MediaPlateEditorModule.cpp
@@ -498,7 +498,7 @@ TSharedRef<FExtender> FMediaPlateEditorModule::ExtendLevelViewportContextMenuFor
 						MenuBuilder.BeginSection("MediaPlate", LOCTEXT("MediaPlateHeading", "Media Plate"));
 						MenuBuilder.AddMenuEntry(
 							LOCTEXT("AddHoldoutComposite", "Apply Holdout Composite"),
-							LOCTEXT("AddHoldoutComposite_Tooltip", "Setup the media plate for alpha holdout compositing to avoid TSR artifacts and the tone curve. Translucency rendered in front is preserved. This technique is preferred but only effective on opaque media currently."),
+							LOCTEXT("AddHoldoutComposite_Tooltip", "Setup the media plate for alpha holdout compositing to avoid TSR artifacts and the tone curve."),
 							FSlateIcon(Style->GetStyleSetName(), "ClassIcon.MediaPlate"),
 							Action_ConfigureHoldoutComposite
 						);
