[MediaPlate] Default aspect ratio is set to 16:9.

diff --git a/Engine/Plugins/Media/MediaPlate/Source/MediaPlate/Private/MediaPlate.cpp b/Engine/Plugins/Media/MediaPlate/Source/MediaPlate/Private/MediaPlate.cpp
--- a/Engine/Plugins/Media/MediaPlate/Source/MediaPlate/Private/MediaPlate.cpp
+++ b/Engine/Plugins/Media/MediaPlate/Source/MediaPlate/Private/MediaPlate.cpp
@@ -61,6 +61,9 @@ AMediaPlate::AMediaPlate(const FObjectInitializer& ObjectInitializer)
 	StaticMeshComponent->bCastDynamicShadow = false;
 	MediaPlateComponent->StaticMeshComponent = StaticMeshComponent;
 
+	// 16:9 by default since most videos are that format.
+	MediaPlateComponent->SetAspectRatio(16.0f/9.0f);
+
 	// Hook up mesh.
 	struct FConstructorStatics
 	{
