[Media] Removing Mips Quantity from Media Texture.

When mipgen is enabled, all mipmaps are generated.

- Deprecated UMediaTexture NumMips property: not used anymore.
- Deprecated FMediaTextureResourceSettings CurrentNumMips property: not used anymore.

diff --git a/Engine/Plugins/Experimental/MediaStream/Source/MediaStream/Private/MediaStreamTextureConfig.cpp b/Engine/Plugins/Experimental/MediaStream/Source/MediaStream/Private/MediaStreamTextureConfig.cpp
--- a/Engine/Plugins/Experimental/MediaStream/Source/MediaStream/Private/MediaStreamTextureConfig.cpp
+++ b/Engine/Plugins/Experimental/MediaStream/Source/MediaStream/Private/MediaStreamTextureConfig.cpp
@@ -4,18 +4,16 @@
 
 bool FMediaStreamTextureConfig::operator==(const FMediaStreamTextureConfig& InOther) const
 {
-	return bEnableMipGen == InOther.bEnableMipGen
-		&& MipGenCount == InOther.MipGenCount;
+	return bEnableMipGen == InOther.bEnableMipGen;
 }
 
 void FMediaStreamTextureConfig::ApplyConfig(UMediaTexture& InMediaTexture) const
 {
-	if (InMediaTexture.EnableGenMips == bEnableMipGen && InMediaTexture.NumMips == MipGenCount)
+	if (InMediaTexture.EnableGenMips == bEnableMipGen)
 	{
 		return;
 	}
 
 	InMediaTexture.EnableGenMips = bEnableMipGen;
-	InMediaTexture.NumMips = MipGenCount;
 	InMediaTexture.UpdateResource();
 }
diff --git a/Engine/Plugins/Experimental/MediaStream/Source/MediaStream/Public/MediaStreamTextureConfig.h b/Engine/Plugins/Experimental/MediaStream/Source/MediaStream/Public/MediaStreamTextureConfig.h
--- a/Engine/Plugins/Experimental/MediaStream/Source/MediaStream/Public/MediaStreamTextureConfig.h
+++ b/Engine/Plugins/Experimental/MediaStream/Source/MediaStream/Public/MediaStreamTextureConfig.h
@@ -16,9 +16,6 @@ struct FMediaStreamTextureConfig
 	UPROPERTY(EditInstanceOnly, BlueprintReadWrite, Category = "Media Stream|Texture")
 	bool bEnableMipGen = false;
 
-	UPROPERTY(EditInstanceOnly, BlueprintReadWrite, Category = "Media Stream|Texture")
-	uint8 MipGenCount = 1;
-
 	MEDIASTREAM_API bool operator==(const FMediaStreamTextureConfig& InOther) const;
 
 	void ApplyConfig(UMediaTexture& InMediaTexture) const;
diff --git a/Engine/Plugins/Media/MediaPlate/Source/MediaPlate/Private/MediaPlateComponent.cpp b/Engine/Plugins/Media/MediaPlate/Source/MediaPlate/Private/MediaPlateComponent.cpp
--- a/Engine/Plugins/Media/MediaPlate/Source/MediaPlate/Private/MediaPlateComponent.cpp
+++ b/Engine/Plugins/Media/MediaPlate/Source/MediaPlate/Private/MediaPlateComponent.cpp
@@ -43,10 +43,9 @@ namespace UE::MediaPlateComponent
 		for (TObjectPtr<UMediaTexture> MediaTexture : MediaTextures)
 		{
 			if (MediaTexture != nullptr
-				&& (MediaTexture->EnableGenMips != MediaTextureSettings.bEnableGenMips || MediaTexture->NumMips != MediaTextureSettings.CurrentNumMips))
+				&& (MediaTexture->EnableGenMips != MediaTextureSettings.bEnableGenMips))
 			{
 				MediaTexture->EnableGenMips = MediaTextureSettings.bEnableGenMips;
-				MediaTexture->NumMips = MediaTextureSettings.CurrentNumMips;
 				MediaTexture->UpdateResource();
 			}
 		}
@@ -59,7 +58,7 @@ namespace UE::MediaPlateComponent
 		{
 			if (MediaTexture != nullptr)
 			{
-				bool bMediaTextureMipGenPropertiesInSync = MediaTexture->EnableGenMips == MediaTextureSettings.bEnableGenMips && MediaTexture->NumMips == MediaTextureSettings.CurrentNumMips;
+				const bool bMediaTextureMipGenPropertiesInSync = MediaTexture->EnableGenMips == MediaTextureSettings.bEnableGenMips;
 				ensureMsgf(bMediaTextureMipGenPropertiesInSync, TEXT("Mip Generation properties set on Media Plate are different from the properties set on Media Texture. \n\
 					Media Texture mip generation properites are not meant to be modified directly."));
 			}
@@ -124,7 +123,6 @@ UMediaPlateComponent::UMediaPlateComponent(const FObjectInitializer& ObjectIniti
 	VisibleMipsTilesCalculations = EMediaTextureVisibleMipsTiles::Plane;
 
 	MediaTextureSettings.bEnableGenMips = false;
-	MediaTextureSettings.CurrentNumMips = 1;
 }
 
 #if WITH_EDITOR
@@ -1150,11 +1148,9 @@ void UMediaPlateComponent::SetUpTextures(UE::MediaPlateComponent::ESetUpTextures
 		{
 			bool bApplyTextureUpdate = false;
 
-			if (MediaTexture->EnableGenMips != MediaTextureSettings.bEnableGenMips
-				|| MediaTexture->NumMips != MediaTextureSettings.CurrentNumMips)
+			if (MediaTexture->EnableGenMips != MediaTextureSettings.bEnableGenMips)
 			{
 				MediaTexture->EnableGenMips = MediaTextureSettings.bEnableGenMips;
-				MediaTexture->NumMips = MediaTextureSettings.CurrentNumMips;
 				bApplyTextureUpdate = true;
 			}
 
@@ -1457,8 +1453,7 @@ void UMediaPlateComponent::PostEditChangeProperty(FPropertyChangedEvent& Propert
 			MediaTextureTrackerObject->MipLevelToUpscale = bEnableMipMapUpscaling ? MipLevelToUpscale : -1;
 		}
 	}
-	else if (PropertyChangedEvent.GetPropertyName() == GET_MEMBER_NAME_CHECKED(FMediaTextureResourceSettings, bEnableGenMips)
-		|| PropertyChangedEvent.GetPropertyName() == GET_MEMBER_NAME_CHECKED(FMediaTextureResourceSettings, CurrentNumMips))
+	else if (PropertyChangedEvent.GetPropertyName() == GET_MEMBER_NAME_CHECKED(FMediaTextureResourceSettings, bEnableGenMips))
 	{
 		UE::MediaPlateComponent::ApplyMediaTextureMipGenProperties(MediaTextureSettings, MediaTextures);
 		RestartPlayer();
diff --git a/Engine/Plugins/Media/MediaPlate/Source/MediaPlate/Public/MediaPlateComponent.h b/Engine/Plugins/Media/MediaPlate/Source/MediaPlate/Public/MediaPlateComponent.h
--- a/Engine/Plugins/Media/MediaPlate/Source/MediaPlate/Public/MediaPlateComponent.h
+++ b/Engine/Plugins/Media/MediaPlate/Source/MediaPlate/Public/MediaPlateComponent.h
@@ -52,9 +52,11 @@ struct FMediaTextureResourceSettings
 	UPROPERTY(EditAnywhere, Category = "MediaTexture", meta = (DisplayName = "Enable RealTime Mips"))
 	bool bEnableGenMips = false;
 
-	/** Current number of mips to be generated as output */
-	UPROPERTY(EditAnywhere, Category = "MediaTexture", meta = (DisplayName = "Mips Quantity"))
-	uint8 CurrentNumMips = 1;
+#if WITH_EDITORONLY_DATA
+	UE_DEPRECATED(5.6, "Only 'Enable RealTime Mips' is needed now. When true, the full mip chain will be generated.")
+	UPROPERTY(meta = (DeprecatedProperty, DeprecationMessage = "Only 'Enable RealTime Mips' is needed now. When true, the full mip chain will be generated."))
+	uint8 CurrentNumMips_DEPRECATED = 1;
+#endif
 };
 
 
diff --git a/Engine/Source/Runtime/MediaAssets/Private/Assets/MediaTexture.cpp b/Engine/Source/Runtime/MediaAssets/Private/Assets/MediaTexture.cpp
--- a/Engine/Source/Runtime/MediaAssets/Private/Assets/MediaTexture.cpp
+++ b/Engine/Source/Runtime/MediaAssets/Private/Assets/MediaTexture.cpp
@@ -77,7 +77,9 @@ UMediaTexture::UMediaTexture(const FObjectInitializer& ObjectInitializer)
 	, AutoClear(false)
 	, ClearColor(FLinearColor::Black)
 	, EnableGenMips(false)
-	, NumMips(1)
+#if WITH_EDITORONLY_DATA
+	, NumMips_DEPRECATED(1)
+#endif
 	, NewStyleOutput(false)
 	, CurrentAspectRatio(0.0f)
 	, CurrentOrientation(MTORI_Original)
@@ -175,8 +177,8 @@ FTextureResource* UMediaTexture::CreateResource()
 		Filter = (TextureNumMips > 1) ? TF_Trilinear : TF_Bilinear;
 	}
 
-	const uint8 ClampedNumMips = FMath::Min(NumMips, static_cast<uint8>(GMaxTextureMipCount));
-	return new FMediaTextureResource(*this, Dimensions, Size, ClearColor, CurrentGuid.IsValid() ? CurrentGuid : DefaultGuid, EnableGenMips, ClampedNumMips, ColorspaceOverride);
+	const uint8 MaxNumMips = static_cast<uint8>(GMaxTextureMipCount);
+	return new FMediaTextureResource(*this, Dimensions, Size, ClearColor, CurrentGuid.IsValid() ? CurrentGuid : DefaultGuid, EnableGenMips, MaxNumMips, ColorspaceOverride);
 }
 
 
@@ -401,6 +403,7 @@ void UMediaTexture::TickResource(FTimespan Timecode)
 
 	LastClearColor = ClearColor;
 	LastSrgb = SRGB;
+	const uint8 MaxNumMips = static_cast<uint8>(GMaxTextureMipCount);
 
 	// set up render parameters
 	FMediaTextureResource::FRenderParams RenderParams;
@@ -440,7 +443,7 @@ void UMediaTexture::TickResource(FTimespan Timecode)
 				SRGB = Sample->IsOutputSrgb();
 				LastSrgb = SRGB;
 
-				TextureNumMips = UE::MediaTexture::Private::CalcTextureNumMips(Sample, NumMips, EnableGenMips);
+				TextureNumMips = UE::MediaTexture::Private::CalcTextureNumMips(Sample, MaxNumMips, EnableGenMips);
 				bIsSampleValid = true;
 			}
 			else
@@ -460,7 +463,7 @@ void UMediaTexture::TickResource(FTimespan Timecode)
 						SRGB = Sample->IsOutputSrgb();
 						LastSrgb = SRGB;
 
-						TextureNumMips = UE::MediaTexture::Private::CalcTextureNumMips(Sample, NumMips, EnableGenMips);
+						TextureNumMips = UE::MediaTexture::Private::CalcTextureNumMips(Sample, MaxNumMips, EnableGenMips);
 						bIsSampleValid = true;
 					}
 
@@ -508,8 +511,8 @@ void UMediaTexture::TickResource(FTimespan Timecode)
 	RenderParams.ClearColor = ClearColor;
 	RenderParams.PreviousGuid = PreviousGuid;
 	RenderParams.CurrentGuid = CurrentGuid;
-	RenderParams.NumMips = FMath::Min(NumMips, static_cast<uint8>(GMaxTextureMipCount));
-	
+	RenderParams.NumMips = MaxNumMips;
+
 	// redraw texture resource on render thread
 	FMediaTextureResource* ResourceParam = (FMediaTextureResource*)GetResource();
 
diff --git a/Engine/Source/Runtime/MediaAssets/Private/Misc/MediaTextureResource.h b/Engine/Source/Runtime/MediaAssets/Private/Misc/MediaTextureResource.h
--- a/Engine/Source/Runtime/MediaAssets/Private/Misc/MediaTextureResource.h
+++ b/Engine/Source/Runtime/MediaAssets/Private/Misc/MediaTextureResource.h
@@ -43,7 +43,7 @@ public:
 	 *
 	 * @param InOwner The Movie texture object to create a resource for (must not be nullptr).
 	 * @param InOwnerDim Reference to the width and height of the texture that owns this resource (will be updated by resource).
-	 * @param InOWnerSize Reference to the size in bytes of the texture that owns this resource (will be updated by resource).
+	 * @param InOwnerSize Reference to the size in bytes of the texture that owns this resource (will be updated by resource).
 	 * @param InClearColor The initial clear color.
 	 * @param InTextureGuid The initial external texture GUID.
 	 * @param bEnableGenMips If true mips generation will be enabled (possibly optimizing for NumMips == 1 case)
diff --git a/Engine/Source/Runtime/MediaAssets/Public/MediaTexture.h b/Engine/Source/Runtime/MediaAssets/Public/MediaTexture.h
--- a/Engine/Source/Runtime/MediaAssets/Public/MediaTexture.h
+++ b/Engine/Source/Runtime/MediaAssets/Public/MediaTexture.h
@@ -81,9 +81,11 @@ public:
 	UPROPERTY(EditAnywhere, BlueprintReadOnly, Category = "MediaTexture", meta = (DisplayName = "Enable Mipmap generation"))
 	bool EnableGenMips;
 
-	/** The number of mips to use (default = 1). */
-	UPROPERTY(EditAnywhere, BlueprintReadWrite, Category = "MediaTexture", meta=(DisplayName="Total number of Mipmaps to output"))
-	uint8 NumMips;
+#if WITH_EDITORONLY_DATA
+	UE_DEPRECATED(5.6, "Only 'Enable Mipmap Generation' is needed now. When true, the full mip chain will be generated.")
+	UPROPERTY(meta = (DeprecatedProperty, DeprecationMessage = "Only 'Enable Mipmap Generation' is needed now. When true, the full mip chain will be generated."))
+	uint8 NumMips_DEPRECATED;
+#endif
 
 	/** Enable new style output. */
 	UPROPERTY(EditAnywhere, BlueprintReadWrite, Category = "MediaTexture", meta = (DisplayName = "Enable new style output"))
