[MediaPlate] Fix CIS issue on Mac.

diff --git a/Engine/Plugins/Media/MediaPlate/Source/MediaPlate/Public/MediaPlateComponent.h b/Engine/Plugins/Media/MediaPlate/Source/MediaPlate/Public/MediaPlateComponent.h
--- a/Engine/Plugins/Media/MediaPlate/Source/MediaPlate/Public/MediaPlateComponent.h
+++ b/Engine/Plugins/Media/MediaPlate/Source/MediaPlate/Public/MediaPlateComponent.h
@@ -57,6 +57,14 @@ struct FMediaTextureResourceSettings
 	UPROPERTY(meta = (DeprecatedProperty, DeprecationMessage = "Only 'Enable RealTime Mips' is needed now. When true, the full mip chain will be generated."))
 	uint8 CurrentNumMips_DEPRECATED = 1;
 #endif
+
+PRAGMA_DISABLE_DEPRECATION_WARNINGS
+	FMediaTextureResourceSettings() = default;
+	FMediaTextureResourceSettings(const FMediaTextureResourceSettings& InOther) = default;
+	FMediaTextureResourceSettings(FMediaTextureResourceSettings&& InOther) = default;
+	FMediaTextureResourceSettings& operator=(const FMediaTextureResourceSettings&) = default;
+	FMediaTextureResourceSettings& operator=(FMediaTextureResourceSettings&&) = default;
+PRAGMA_ENABLE_DEPRECATION_WARNINGS
 };
 
 
