MediaPlate: Remove deprecated overlay comp menu option, though we keep the materials & underlying functionality present for now.

diff --git a/Engine/Plugins/Media/MediaPlate/Source/MediaPlateEditor/Private/MediaPlateEditorModule.cpp b/Engine/Plugins/Media/MediaPlate/Source/MediaPlateEditor/Private/MediaPlateEditorModule.cpp
--- a/Engine/Plugins/Media/MediaPlate/Source/MediaPlateEditor/Private/MediaPlateEditorModule.cpp
+++ b/Engine/Plugins/Media/MediaPlate/Source/MediaPlateEditor/Private/MediaPlateEditorModule.cpp
@@ -430,52 +430,6 @@ TSharedRef<FExtender> FMediaPlateEditorModule::ExtendLevelViewportContextMenuFor
 								, MediaPlateActor)
 						);
 
-						// Note: Could be removed now that we have the holdout composite alternative.
-						FUIAction Action_ConfigureOverlayComposite(
-							FExecuteAction::CreateLambda([](AMediaPlate* InMediaPlateActor)
-								{
-									if (!IsValid(InMediaPlateActor))
-									{
-										return;
-									}
-
-									const FScopedTransaction Transaction(LOCTEXT("ApplyOverlayCompositeMats", "Apply Overlay Composite Materials (Deprecated)"));
-									InMediaPlateActor->Modify();
-
-									UMaterial* BasePassMaterial = LoadObject<UMaterial>(NULL, TEXT("/MediaPlate/M_MediaPlate_Masked"), NULL, LOAD_None, NULL);
-									InMediaPlateActor->ApplyMaterial(BasePassMaterial);
-
-									UMaterial* OverlayMaterial = LoadObject<UMaterial>(NULL, TEXT("/MediaPlate/M_MediaPlate_OverlayComp"), NULL, LOAD_None, NULL);
-									InMediaPlateActor->ApplyOverlayMaterial(OverlayMaterial);
-
-									UE::MediaPlate::Private::ApplyTranslucencyScreenPercentageCVar(1);
-								}
-								, MediaPlateActor),
-							FCanExecuteAction::CreateLambda([](AMediaPlate* InMediaPlateActor) -> bool
-								{
-									return IsValid(InMediaPlateActor) ? !InMediaPlateActor->FindComponentByClass<UHoldoutCompositeComponent>() : false;
-								}
-								, MediaPlateActor)
-						);
-
-						FUIAction Action_ResetDefault(
-							FExecuteAction::CreateLambda([](AMediaPlate* InMediaPlateActor)
-								{
-									if (!IsValid(InMediaPlateActor))
-									{
-										return;
-									}
-
-									const FScopedTransaction Transaction(LOCTEXT("ResetDefaultMats", "Reset Default Materials"));
-
-									InMediaPlateActor->Modify();
-									InMediaPlateActor->UseDefaultMaterial();
-
-									UE::MediaPlate::Private::ApplyTranslucencyScreenPercentageCVar(0);
-								}
-								, MediaPlateActor)
-						);
-
 						FUIAction Action_CopyMediaTextureReference(
 							FExecuteAction::CreateLambda([](AMediaPlate* InMediaPlateActor)
 								{
@@ -508,19 +462,6 @@ TSharedRef<FExtender> FMediaPlateEditorModule::ExtendLevelViewportContextMenuFor
 							FSlateIcon(Style->GetStyleSetName(), "ClassIcon.MediaPlate"),
 							Action_RemoveHoldoutComposite
 						);
-						MenuBuilder.AddSeparator();
-						MenuBuilder.AddMenuEntry(
-							LOCTEXT("ApplyOverlayCompositeMats", "Apply Overlay Composite Materials (Deprecated)"),
-							LOCTEXT("ApplyOverlayCompositeMats_Tooltip", "Setup the media plate for overlay compositing to avoid TSR artifacts by replacing relevant materials. This technique is only effective on opaque media currently."),
-							FSlateIcon(Style->GetStyleSetName(), "ClassIcon.MediaPlate"),
-							Action_ConfigureOverlayComposite
-						);
-						MenuBuilder.AddMenuEntry(
-							LOCTEXT("ResetDefaultMats", "Reset Default Materials"),
-							LOCTEXT("ResetDefaultMats_Tooltip", "Reverts the media plate to its default materials. Note that we also globally disable r.Translucency.ScreenPercentage.Basis if previously enabled."),
-							FSlateIcon(Style->GetStyleSetName(), "ClassIcon.MediaPlate"),
-							Action_ResetDefault
-						);
 
 						MenuBuilder.AddSeparator();
 						MenuBuilder.AddMenuEntry(
