[MediaPlate] Don't close the media player when modifying any properties. We only want to do that for specific properties.

OnUnregister() is called several times, in unregister - register call sequences, from within the UActorComponent::PostEditChangeProperty and AActor::PostEditChangeProperty, so it is not the right place to close the media player.

To fix this, the Close() call is moved to OnComponentDestroyed() which is called when an actor is deleted. Tested the repro steps of UE-192347 that this was originally fixing and confirmed that it doesn't reintroduce the problem.

diff --git a/Engine/Plugins/Media/MediaPlate/Source/MediaPlate/Private/MediaPlateComponent.cpp b/Engine/Plugins/Media/MediaPlate/Source/MediaPlate/Private/MediaPlateComponent.cpp
--- a/Engine/Plugins/Media/MediaPlate/Source/MediaPlate/Private/MediaPlateComponent.cpp
+++ b/Engine/Plugins/Media/MediaPlate/Source/MediaPlate/Private/MediaPlateComponent.cpp
@@ -215,6 +215,9 @@ void UMediaPlateComponent::BeginDestroy()
 		ClockSink.Reset();
 	}
 
+	// Ensure media resources are properly released.
+	Close();
+	
 	Super::BeginDestroy();
 }
 
@@ -280,8 +283,15 @@ void UMediaPlateComponent::OnUnregister()
 	Super::OnUnregister();
 
 	UnregisterWithMediaTextureTracker();
+}
 
+void UMediaPlateComponent::OnComponentDestroyed(bool bDestroyingHierarchy)
+{
+	// Free the resources as soon as the component is destroyed (sooner than GC).
+	// If waiting for GC, it looks like resources are leaking.
 	Close();
+
+	Super::OnComponentDestroyed(bDestroyingHierarchy);
 }
 
 UMediaPlayer* UMediaPlateComponent::GetMediaPlayer()
diff --git a/Engine/Plugins/Media/MediaPlate/Source/MediaPlate/Public/MediaPlateComponent.h b/Engine/Plugins/Media/MediaPlate/Source/MediaPlate/Public/MediaPlateComponent.h
--- a/Engine/Plugins/Media/MediaPlate/Source/MediaPlate/Public/MediaPlateComponent.h
+++ b/Engine/Plugins/Media/MediaPlate/Source/MediaPlate/Public/MediaPlateComponent.h
@@ -87,6 +87,7 @@ public:
 	virtual void BeginDestroy() override;
 	virtual void TickComponent(float DeltaTime, enum ELevelTick TickType, FActorComponentTickFunction* ThisTickFunction) override;
 	virtual void OnUnregister() override;
+	virtual void OnComponentDestroyed(bool bDestroyingHierarchy) override;
 #if WITH_EDITOR
 	virtual void PostEditChangeProperty(FPropertyChangedEvent& PropertyChangedEvent) override;
 #endif // WITH_EDITOR
