MediaPlate: Holdout composite is now a checkbox, instead of a right-click option that adds a component.

diff --git a/Engine/Plugins/Media/MediaPlate/Source/MediaPlate/MediaPlate.Build.cs b/Engine/Plugins/Media/MediaPlate/Source/MediaPlate/MediaPlate.Build.cs
--- a/Engine/Plugins/Media/MediaPlate/Source/MediaPlate/MediaPlate.Build.cs
+++ b/Engine/Plugins/Media/MediaPlate/Source/MediaPlate/MediaPlate.Build.cs
@@ -23,6 +23,7 @@ public class MediaPlate : ModuleRules
 				"Engine",
 				"MediaAssets",
 				"MediaUtils",
+				"HoldoutComposite"
 			}
 			);
 
diff --git a/Engine/Plugins/Media/MediaPlate/Source/MediaPlate/Private/MediaPlate.cpp b/Engine/Plugins/Media/MediaPlate/Source/MediaPlate/Private/MediaPlate.cpp
--- a/Engine/Plugins/Media/MediaPlate/Source/MediaPlate/Private/MediaPlate.cpp
+++ b/Engine/Plugins/Media/MediaPlate/Source/MediaPlate/Private/MediaPlate.cpp
@@ -4,6 +4,9 @@
 
 #include "Components/StaticMeshComponent.h"
 #include "Engine/StaticMesh.h"
+#include "Engine/World.h"
+#include "HoldoutCompositeComponent.h"
+#include "HoldoutCompositeSubsystem.h"
 #include "Materials/Material.h"
 #include "Materials/MaterialInstanceConstant.h"
 #include "Materials/MaterialInstanceDynamic.h"
@@ -11,6 +14,7 @@
 #include "MediaPlateModule.h"
 #include "MediaTexture.h"
 #include "UObject/ConstructorHelpers.h"
+#include "UObject/FortniteMainBranchObjectVersion.h"
 
 #if WITH_EDITOR
 #include "Editor.h"
@@ -27,6 +31,7 @@ FLazyName AMediaPlate::MediaTextureName("MediaTexture");
 
 namespace UE::MediaPlate::Private
 {
+	// TODO: Could be removed (with associated materials) now that we've deprecated overlay compositing.
 	void ApplyTranslucencyScreenPercentageCVar(int32 InBasis)
 	{
 		static IConsoleVariable* TranslucencySPBasisCVar = IConsoleManager::Get().FindConsoleVariable(TEXT("r.Translucency.ScreenPercentage.Basis"));
@@ -90,6 +95,27 @@ AMediaPlate::AMediaPlate(const FObjectInitializer& ObjectInitializer)
 #endif // WITH_EDITOR
 }
 
+void AMediaPlate::PostLoad()
+{
+	Super::PostLoad();
+
+#if WITH_EDITORONLY_DATA
+	if (!HasAnyFlags(RF_ClassDefaultObject))
+	{
+		// The holdout composite component was replaced with a checkbox.
+		if (GetLinkerCustomVersion(FFortniteMainBranchObjectVersion::GUID) < FFortniteMainBranchObjectVersion::MediaPlateHoldoutComponentRemoval)
+		{
+			if (UHoldoutCompositeComponent* HoldoutCompositeComponent = FindComponentByClass<UHoldoutCompositeComponent>())
+			{
+				bEnableHoldoutComposite = HoldoutCompositeComponent->IsEnabled();
+
+				HoldoutCompositeComponent->DestroyComponent();
+			}
+		}
+	}
+#endif
+}
+
 void AMediaPlate::PostActorCreated()
 {
 	Super::PostActorCreated();
@@ -123,10 +149,14 @@ void AMediaPlate::PostRegisterAllComponents()
 	AddAssetUserData();
 #endif // WITH_EDITOR
 
-	if (IsValid(StaticMeshComponent) &&
-		IsValid(StaticMeshComponent->OverlayMaterial))
+	if (IsValid(StaticMeshComponent))
 	{
-		UE::MediaPlate::Private::ApplyTranslucencyScreenPercentageCVar(1);
+		if (IsValid(StaticMeshComponent->OverlayMaterial))
+		{
+			UE::MediaPlate::Private::ApplyTranslucencyScreenPercentageCVar(1);
+		}
+
+		ConditionallyEnableHoldoutComposite();
 	}
 }
 
@@ -135,6 +165,31 @@ void AMediaPlate::BeginDestroy()
 	Super::BeginDestroy();
 }
 
+void AMediaPlate::SetHoldoutCompositeEnabled(bool bInEnabled)
+{
+	bEnableHoldoutComposite = bInEnabled;
+
+	UHoldoutCompositeSubsystem* HoldoutCompositeSubsystem = UWorld::GetSubsystem<UHoldoutCompositeSubsystem>(GetWorld());
+	if (IsValid(HoldoutCompositeSubsystem))
+	{
+		if (bEnableHoldoutComposite)
+		{
+			// Note: re-registering the same component is safe.
+			HoldoutCompositeSubsystem->RegisterPrimitive(StaticMeshComponent);
+		}
+		else
+		{
+			// Note: destroyed components are also automatically removed by the system.
+			HoldoutCompositeSubsystem->UnregisterPrimitive(StaticMeshComponent);
+		}
+	}
+}
+
+bool AMediaPlate::IsHoldoutCompositeEnabled() const
+{
+	return bEnableHoldoutComposite;
+}
+
 UMaterialInterface* AMediaPlate::GetCurrentMaterial() const
 {
 	if (IsValid(StaticMeshComponent))
@@ -155,8 +210,33 @@ UMaterialInterface* AMediaPlate::GetCurrentOverlayMaterial() const
 	return nullptr;
 }
 
+void AMediaPlate::ConditionallyEnableHoldoutComposite()
+{
+	if (bEnableHoldoutComposite)
+	{
+		SetHoldoutCompositeEnabled(true);
+	}
+}
+
 #if WITH_EDITOR
 
+void AMediaPlate::PostEditChangeProperty(FPropertyChangedEvent& PropertyChangedEvent)
+{
+	Super::PostEditChangeProperty(PropertyChangedEvent);
+
+	if (PropertyChangedEvent.GetPropertyName() == GET_MEMBER_NAME_CHECKED(ThisClass, bEnableHoldoutComposite))
+	{
+		SetHoldoutCompositeEnabled(bEnableHoldoutComposite);
+	}
+}
+
+void AMediaPlate::OnStaticMeshChange()
+{
+	ApplyCurrentMaterial();
+
+	ConditionallyEnableHoldoutComposite();
+}
+
 void AMediaPlate::UseDefaultMaterial()
 {
 	UMaterial* DefaultMaterial = LoadObject<UMaterial>(NULL, TEXT("/MediaPlate/M_MediaPlate_Opaque"), NULL, LOAD_None, NULL);
@@ -369,7 +449,7 @@ void AMediaPlate::AddAssetUserData()
 	if (StaticMeshComponent != nullptr && !StaticMeshComponent->HasAssetUserDataOfClass(UMediaPlateAssetUserData::StaticClass()))
 	{
 		UMediaPlateAssetUserData* AssetUserData = NewObject<UMediaPlateAssetUserData>(GetTransientPackage());
-		AssetUserData->OnPostEditChangeOwner.BindUObject(this, &AMediaPlate::ApplyCurrentMaterial);
+		AssetUserData->OnPostEditChangeOwner.BindUObject(this, &AMediaPlate::OnStaticMeshChange);
 		StaticMeshComponent->AddAssetUserData(AssetUserData);
 	}
 }
@@ -382,7 +462,7 @@ void AMediaPlate::RemoveAssetUserData()
 	}
 }
 
-#endif
+#endif // WITH_EDITOR
 
 #undef LOCTEXT_NAMESPACE
 
diff --git a/Engine/Plugins/Media/MediaPlate/Source/MediaPlate/Public/MediaPlate.h b/Engine/Plugins/Media/MediaPlate/Source/MediaPlate/Public/MediaPlate.h
--- a/Engine/Plugins/Media/MediaPlate/Source/MediaPlate/Public/MediaPlate.h
+++ b/Engine/Plugins/Media/MediaPlate/Source/MediaPlate/Public/MediaPlate.h
@@ -11,11 +11,6 @@ class UMaterialInterface;
 class UMaterialInstanceDynamic;
 class UMediaPlateComponent;
 
-namespace UE::MediaPlate::Private
-{
-	MEDIAPLATE_API void ApplyTranslucencyScreenPercentageCVar(int32 InBasis);
-}
-
 /**
  * MediaPlate is an actor that can play and show media in the world.
  */
@@ -25,6 +20,10 @@ class MEDIAPLATE_API AMediaPlate : public AActor
 	GENERATED_UCLASS_BODY()
 
 public:
+	//~ Begin UObject Interface.
+	virtual void PostLoad() override;
+	//~ End UObject Interface.
+
 	//~ Begin AActor Interface
 	virtual void PostActorCreated();
 	virtual void PostRegisterAllComponents() override;
@@ -38,9 +37,23 @@ public:
 	UPROPERTY(Category = MediaPlate, VisibleAnywhere, BlueprintReadOnly)
 	TObjectPtr<UStaticMeshComponent> StaticMeshComponent;
 
+	/* Set the holdout composite state. */
+	UFUNCTION(BlueprintSetter)
+	void SetHoldoutCompositeEnabled(bool bInEnabled);
+
+	/* Get the holdout composite state. */
+	UFUNCTION(BlueprintGetter)
+	bool IsHoldoutCompositeEnabled() const;
+
 #if WITH_EDITOR
+	virtual void PostEditChangeProperty(FPropertyChangedEvent& PropertyChangedEvent) override;
 
-	/*
+	/**
+	 * Callback upon static mesh component change (see UMediaPlateAssetUserData).
+	 */
+	void OnStaticMeshChange();
+
+	/**
 	 * Call this to change the static mesh to use the default media plate material and reset the overlay material.
 	 */
 	void UseDefaultMaterial();
@@ -76,6 +89,13 @@ public:
 	UMaterialInterface* GetCurrentOverlayMaterial() const;
 
 private:
+	/** If true, the mesh is rendered separately and composited after post-processing (see HoldoutComposite plugin). Using mip generation with this setting is also recommended if the cost is acceptable, given that the composite bypasses temporal anti-aliasing. */
+	UPROPERTY(EditAnywhere, BlueprintGetter = IsHoldoutCompositeEnabled, BlueprintSetter = SetHoldoutCompositeEnabled, Category = "Materials", meta = (AllowPrivateAccess = true))
+	bool bEnableHoldoutComposite = false;
+
+	/** Only enable registration with holdout composite subsytem if users have enabled holdout composite on the media plate.*/
+	void ConditionallyEnableHoldoutComposite();
+
 	/** Name for our media plate component. */
 	static FLazyName MediaPlateComponentName;
 	/** Name for the media texture parameter in the material. */
diff --git a/Engine/Plugins/Media/MediaPlate/Source/MediaPlateEditor/MediaPlateEditor.Build.cs b/Engine/Plugins/Media/MediaPlate/Source/MediaPlateEditor/MediaPlateEditor.Build.cs
--- a/Engine/Plugins/Media/MediaPlate/Source/MediaPlateEditor/MediaPlateEditor.Build.cs
+++ b/Engine/Plugins/Media/MediaPlate/Source/MediaPlateEditor/MediaPlateEditor.Build.cs
@@ -42,8 +42,7 @@ public class MediaPlateEditor : ModuleRules
 				"Sequencer",
 				"Slate",
 				"SlateCore",
-				"UnrealEd",
-				"HoldoutComposite"
+				"UnrealEd"
 			}
 		);
 	}
diff --git a/Engine/Plugins/Media/MediaPlate/Source/MediaPlateEditor/Private/MediaPlateEditorModule.cpp b/Engine/Plugins/Media/MediaPlate/Source/MediaPlateEditor/Private/MediaPlateEditorModule.cpp
--- a/Engine/Plugins/Media/MediaPlate/Source/MediaPlateEditor/Private/MediaPlateEditorModule.cpp
+++ b/Engine/Plugins/Media/MediaPlate/Source/MediaPlateEditor/Private/MediaPlateEditorModule.cpp
@@ -5,7 +5,6 @@
 #include "Editor.h"
 #include "EngineAnalytics.h"
 #include "HAL/PlatformApplicationMisc.h"
-#include "HoldoutCompositeComponent.h"
 #include "IPlacementModeModule.h"
 #include "ISequencerModule.h"
 #include "LevelEditor.h"
@@ -352,84 +351,6 @@ TSharedRef<FExtender> FMediaPlateEditorModule::ExtendLevelViewportContextMenuFor
 						const ISlateStyle* Style = &FMediaPlateEditorStyle::Get().Get();
 						check(Style);
 
-						FUIAction Action_ConfigureHoldoutComposite(
-							FExecuteAction::CreateLambda([](AMediaPlate* InMediaPlateActor)
-								{
-									if (!IsValid(InMediaPlateActor))
-									{
-										return;
-									}
-
-									if (!IsValid(InMediaPlateActor->StaticMeshComponent))
-									{
-										UE_LOG(LogMediaPlateEditor, Warning, TEXT("Could not setup holdout composite on media plate, invalid static mesh component."));
-										return;
-									}
-
-									if (InMediaPlateActor->FindComponentByClass<UHoldoutCompositeComponent>() != nullptr)
-									{
-										// Already present, do nothing
-										return;
-									}
-
-									const FScopedTransaction Transaction(LOCTEXT("ApplyHoldoutComposite", "Apply Holdout Composite"));
-									InMediaPlateActor->Modify();
-
-									UHoldoutCompositeComponent* HoldoutComponent = NewObject<UHoldoutCompositeComponent>(InMediaPlateActor
-										, MakeUniqueObjectName(InMediaPlateActor, UHoldoutCompositeComponent::StaticClass())
-										, RF_Transactional);
-
-									InMediaPlateActor->AddInstanceComponent(HoldoutComponent);
-
-									HoldoutComponent->SetupAttachment(InMediaPlateActor->StaticMeshComponent);
-									HoldoutComponent->OnComponentCreated();
-									HoldoutComponent->RegisterComponent();
-
-									// Used to update the sub-object tree-view list
-									FLevelEditorModule& LevelEditor = FModuleManager::LoadModuleChecked<FLevelEditorModule>("LevelEditor");
-									LevelEditor.BroadcastComponentsEdited();
-								}
-								, MediaPlateActor),
-							FCanExecuteAction::CreateLambda([](AMediaPlate* InMediaPlateActor) -> bool
-								{
-									// Disable when the media plate has an overlay material or an existing holdout
-									return IsValid(InMediaPlateActor) ? !IsValid(InMediaPlateActor->GetCurrentOverlayMaterial()) && !InMediaPlateActor->FindComponentByClass<UHoldoutCompositeComponent>() : false;
-								}
-								, MediaPlateActor)
-						);
-
-						FUIAction Action_RemoveHoldoutComposite(
-							FExecuteAction::CreateLambda([](AMediaPlate* InMediaPlateActor)
-								{
-									if (!IsValid(InMediaPlateActor))
-									{
-										return;
-									}
-
-									UHoldoutCompositeComponent* HoldoutComponent = InMediaPlateActor->FindComponentByClass<UHoldoutCompositeComponent>();
-									if (IsValid(HoldoutComponent))
-									{
-										const FScopedTransaction Transaction(LOCTEXT("RemoveHoldoutComposite", "Remove Holdout Composite"));
-
-										InMediaPlateActor->RemoveInstanceComponent(HoldoutComponent);
-
-										HoldoutComponent->DetachFromComponent(FDetachmentTransformRules::KeepRelativeTransform);
-										HoldoutComponent->UnregisterComponent();
-										HoldoutComponent->DestroyComponent();
-
-										// Used to update the sub-object tree-view list
-										FLevelEditorModule& LevelEditor = FModuleManager::LoadModuleChecked<FLevelEditorModule>("LevelEditor");
-										LevelEditor.BroadcastComponentsEdited();
-									}
-								}
-								, MediaPlateActor),
-							FCanExecuteAction::CreateLambda([](AMediaPlate* InMediaPlateActor) -> bool
-								{
-									return IsValid(InMediaPlateActor) ? InMediaPlateActor->FindComponentByClass<UHoldoutCompositeComponent>() != nullptr : false;
-								}
-								, MediaPlateActor)
-						);
-
 						FUIAction Action_CopyMediaTextureReference(
 							FExecuteAction::CreateLambda([](AMediaPlate* InMediaPlateActor)
 								{
@@ -450,20 +371,6 @@ TSharedRef<FExtender> FMediaPlateEditorModule::ExtendLevelViewportContextMenuFor
 						);
 
 						MenuBuilder.BeginSection("MediaPlate", LOCTEXT("MediaPlateHeading", "Media Plate"));
-						MenuBuilder.AddMenuEntry(
-							LOCTEXT("AddHoldoutComposite", "Apply Holdout Composite"),
-							LOCTEXT("AddHoldoutComposite_Tooltip", "Setup the media plate for alpha holdout compositing to avoid TSR artifacts and the tone curve."),
-							FSlateIcon(Style->GetStyleSetName(), "ClassIcon.MediaPlate"),
-							Action_ConfigureHoldoutComposite
-						);
-						MenuBuilder.AddMenuEntry(
-							LOCTEXT("RemoveHoldoutComposite", "Remove Holdout Composite"),
-							LOCTEXT("RemoveHoldoutComposite_Tooltip", "Remove and destroy the holdout composite component, if present."),
-							FSlateIcon(Style->GetStyleSetName(), "ClassIcon.MediaPlate"),
-							Action_RemoveHoldoutComposite
-						);
-
-						MenuBuilder.AddSeparator();
 						MenuBuilder.AddMenuEntry(
 							LOCTEXT("CopyMediaTextureReference", "Copy Media Texture Reference"),
 							LOCTEXT("CopyMediaTextureReference_Tooltip", "Copy Embedded Media Texture Reference to Clipboard."),
diff --git a/Engine/Source/Runtime/Core/Public/UObject/FortniteMainBranchObjectVersions.inl b/Engine/Source/Runtime/Core/Public/UObject/FortniteMainBranchObjectVersions.inl
--- a/Engine/Source/Runtime/Core/Public/UObject/FortniteMainBranchObjectVersions.inl
+++ b/Engine/Source/Runtime/Core/Public/UObject/FortniteMainBranchObjectVersions.inl
@@ -580,5 +580,8 @@ FORTNITE_MAIN_VERSION(AddMaterialSubstrateSubsurfaceType, 188)
 // Added option to visualize runtime virtual textures' streamed mips only in PIE 
 FORTNITE_MAIN_VERSION(AddedRuntimeVirtualTextureUseStreamingMipsInEditorMode, 189)
 
+// Media plate holdout composite components have been replaced by a checkbox
+FORTNITE_MAIN_VERSION(MediaPlateHoldoutComponentRemoval, 190)
+
 // -----<new versions can be added above this line>-------------------------------------------------
 #undef FORTNITE_MAIN_VERSION
\ No newline at end of file
