[MediaPlate] Fix for: Toggling Realtime Mips checkbox resets the playhead to time 0.

Previous behavior was to restart the player, only if it was playing. It doesn't really serve any purpose to do that.
The mip gen affects only the MediaTexture only and gets updated by the "UpdateResource" called in ApplyMediaTextureMipGenProperties.
If the media player is playing, it will correctly render the next frames of the video without having to be reset.

DisplayCluster + Multi-user note:
Should not be affected by this change. The PostEditChangeProperty will be replicated by multi-user leading to the media texture(s) being updated on every nDisplay nodes.

#rnx

diff --git a/Engine/Plugins/Media/MediaPlate/Source/MediaPlate/Private/MediaPlateComponent.cpp b/Engine/Plugins/Media/MediaPlate/Source/MediaPlate/Private/MediaPlateComponent.cpp
--- a/Engine/Plugins/Media/MediaPlate/Source/MediaPlate/Private/MediaPlateComponent.cpp
+++ b/Engine/Plugins/Media/MediaPlate/Source/MediaPlate/Private/MediaPlateComponent.cpp
@@ -1475,7 +1475,6 @@ void UMediaPlateComponent::PostEditChangeProperty(FPropertyChangedEvent& Propert
 	else if (PropertyChangedEvent.GetPropertyName() == GET_MEMBER_NAME_CHECKED(FMediaTextureResourceSettings, bEnableGenMips))
 	{
 		UE::MediaPlateComponent::ApplyMediaTextureMipGenProperties(MediaTextureSettings, MediaTextures);
-		RestartPlayer();
 	}
 	else if (PropertyChangedEvent.GetPropertyName() == GET_MEMBER_NAME_CHECKED(ThisClass, MediaPlateResource))
 	{
