[MediaPlate] Audio does not stop when adding to Sequencer.

- Stop the MediaPlate's player when adding in a media track since it is taken over by the media track's player.
- Disable audio in the component. Removes the audio icon in the viewport.

bEnableAudio undo:
Trying to salvage the restored component because the audio settings (in the component) would otherwise be lost. It seems the component gets marked as garbage (while in the transcation buffer?) and needs to clear it otherwise it doesn't register properly again.

diff --git a/Engine/Plugins/Media/MediaPlate/Source/MediaPlate/Private/MediaPlateComponent.cpp b/Engine/Plugins/Media/MediaPlate/Source/MediaPlate/Private/MediaPlateComponent.cpp
--- a/Engine/Plugins/Media/MediaPlate/Source/MediaPlate/Private/MediaPlateComponent.cpp
+++ b/Engine/Plugins/Media/MediaPlate/Source/MediaPlate/Private/MediaPlateComponent.cpp
@@ -491,6 +491,15 @@ void UMediaPlateComponent::SetMediaPlateResource(const FMediaPlateResource& InMe
 	RefreshMediaPlateResource();
 }
 
+void UMediaPlateComponent::SetEnableAudio(bool bInEnableAudio)
+{
+	if (bEnableAudio != bInEnableAudio)
+	{
+		bEnableAudio = bInEnableAudio;
+		RefreshMediaSoundComponent();
+	}
+}
+
 void UMediaPlateComponent::SetMeshRange(FVector2D InMeshRange)
 {
 	MeshRange = InMeshRange;
@@ -1387,6 +1396,73 @@ void UMediaPlateComponent::RefreshMediaPlateResource()
 	}
 }
 
+void UMediaPlateComponent::RefreshMediaSoundComponent()
+{
+	// Are we turning on audio?
+	if (bEnableAudio)
+	{
+		// Create the component if missing.
+		if (!SoundComponent)
+		{
+			SoundComponent = CreateMediaSoundComponent();
+		}
+
+		if (SoundComponent)
+		{
+			SoundComponent->SetMediaPlayer(MediaPlayer);
+
+			if (!SoundComponent->IsRegistered())
+			{
+				// If called post undo: the component is marked as garbage and will not register again
+				// unless we clear the garbage flag.
+				SoundComponent->ClearGarbage();
+				SoundComponent->bAutoRegister = true;	// Restore auto-register (if the component is restored from undo)
+				SoundComponent->Initialize();
+				SoundComponent->RegisterComponent();
+			}
+		}
+	}
+	else
+	{
+		// Note that removing the component loses all settings that where done to it.
+		// This is why in the undo, we try to salvage the restored component rather than creating a new one.
+		RemoveMediaSoundComponent();
+	}
+}
+
+UMediaSoundComponent* UMediaPlateComponent::CreateMediaSoundComponent()
+{
+	// Media player must be valid.
+	if (!MediaPlayer)
+	{
+		UE_LOG(LogMediaPlate, Error, TEXT("Failed to create MediaSoundComponent: Missing Media Player."));
+		return nullptr;
+	}
+	
+	// Create a sound component.
+	if (UMediaSoundComponent* NewSoundComponent = NewObject<UMediaSoundComponent>(this, NAME_None))
+	{
+		NewSoundComponent->bIsUISound = true;
+		return NewSoundComponent;
+	}
+			
+	UE_LOG(LogMediaPlate, Error, TEXT("Failed to create MediaSoundComponent."));
+	return nullptr;
+}
+
+void UMediaPlateComponent::RemoveMediaSoundComponent()
+{
+	if (SoundComponent != nullptr)
+	{
+		SoundComponent->UnregisterComponent();
+		SoundComponent->SetMediaPlayer(nullptr);
+		SoundComponent->UpdatePlayer();
+		SoundComponent->DestroyComponent();
+		SoundComponent->bAutoRegister = false;	// Don't let the component register again.
+		SoundComponent = nullptr;
+	}
+}
+
 #if WITH_EDITOR
 
 void UMediaPlateComponent::SetVisibleMipsTilesCalculations(EMediaTextureVisibleMipsTiles InVisibleMipsTilesCalculations)
@@ -1408,35 +1484,7 @@ void UMediaPlateComponent::PostEditChangeProperty(FPropertyChangedEvent& Propert
 	// Has bEnableAudio changed?
 	if (PropertyChangedEvent.GetPropertyName() == GET_MEMBER_NAME_CHECKED(ThisClass, bEnableAudio))
 	{
-		// Are we turning on audio?
-		if (bEnableAudio)
-		{
-			// Get the media player.
-			if (MediaPlayer != nullptr)
-			{
-				// Create a sound component.
-				SoundComponent = NewObject<UMediaSoundComponent>(this, NAME_None);
-				if (SoundComponent != nullptr)
-				{
-					SoundComponent->bIsUISound = true;
-					SoundComponent->SetMediaPlayer(MediaPlayer);
-					SoundComponent->Initialize();
-					SoundComponent->RegisterComponent();
-				}
-			}
-		}
-		else
-		{
-			// Remove this sound component.
-			if (SoundComponent != nullptr)
-			{
-				SoundComponent->UnregisterComponent();
-				SoundComponent->SetMediaPlayer(nullptr);
-				SoundComponent->UpdatePlayer();
-				SoundComponent->DestroyComponent();
-				SoundComponent = nullptr;
-			}
-		}
+		RefreshMediaSoundComponent();
 	}
 	else if (PropertyChangedEvent.GetPropertyName() == GET_MEMBER_NAME_CHECKED(ThisClass, bPlayOnlyWhenVisible))
 	{
@@ -1481,6 +1529,14 @@ void UMediaPlateComponent::PostEditChangeProperty(FPropertyChangedEvent& Propert
 		RefreshMediaPlateResource();
 	}
 }
+
+void UMediaPlateComponent::PostEditUndo()
+{
+	Super::PostEditUndo();
+	
+	// Try to salvage the restored component. 
+	RefreshMediaSoundComponent();
+}
 #endif // WITH_EDITOR
 
 void UMediaPlateComponent::SwitchStates(EMediaPlateEventState State)
diff --git a/Engine/Plugins/Media/MediaPlate/Source/MediaPlate/Public/MediaPlateComponent.h b/Engine/Plugins/Media/MediaPlate/Source/MediaPlate/Public/MediaPlateComponent.h
--- a/Engine/Plugins/Media/MediaPlate/Source/MediaPlate/Public/MediaPlateComponent.h
+++ b/Engine/Plugins/Media/MediaPlate/Source/MediaPlate/Public/MediaPlateComponent.h
@@ -90,6 +90,7 @@ public:
 	virtual void OnComponentDestroyed(bool bDestroyingHierarchy) override;
 #if WITH_EDITOR
 	virtual void PostEditChangeProperty(FPropertyChangedEvent& PropertyChangedEvent) override;
+	virtual void PostEditUndo() override;
 #endif // WITH_EDITOR
 
 	/**
@@ -237,6 +238,9 @@ public:
 	UPROPERTY(EditAnywhere, Category = "Cache", meta = (DisplayName = "Cache", ShowOnlyInnerProperties))
 	FMediaSourceCacheSettings CacheSettings;
 
+	UFUNCTION(BlueprintCallable, Category = "Media|MediaPlateComponent")
+	void SetEnableAudio(bool bInEnableAudio);
+
 	/** Set the arc size in degrees used for visible mips and tiles calculations, specific to the sphere. */
 	UFUNCTION(BlueprintCallable, Category = "Media|MediaPlateComponent")
 	void SetMeshRange(FVector2D InMeshRange);
@@ -595,4 +599,20 @@ private:
 	 * Can be called after updating MediaPlayerResource, to ensure the Playlist and player state are up to date
 	 */
 	void RefreshMediaPlateResource();
+
+	/**
+	 * Refresh the media sound component according to the bEnableAudio variable.
+	 */
+	void RefreshMediaSoundComponent();
+
+	/**
+	 * Create the media sound component with initial setup (but not registered).
+	 * @return created component on success, nullptr on failure.
+	 */
+	UMediaSoundComponent* CreateMediaSoundComponent();
+
+	/**
+	 * Removes the audio component.
+	 */
+	void RemoveMediaSoundComponent();
 };
diff --git a/Engine/Plugins/Media/MediaPlate/Source/MediaPlateEditor/Private/Sequencer/MediaPlateTrackEditor.cpp b/Engine/Plugins/Media/MediaPlate/Source/MediaPlateEditor/Private/Sequencer/MediaPlateTrackEditor.cpp
--- a/Engine/Plugins/Media/MediaPlate/Source/MediaPlateEditor/Private/Sequencer/MediaPlateTrackEditor.cpp
+++ b/Engine/Plugins/Media/MediaPlate/Source/MediaPlateEditor/Private/Sequencer/MediaPlateTrackEditor.cpp
@@ -182,6 +182,9 @@ void FMediaPlateTrackEditor::AddTrackForComponent(UMediaPlateComponent* Componen
 		}
 	}
 
+	// Make sure the media plate's player is not playing
+	Component->Close();
+
 	// Does the media plate have autoplay?
 	if (Component->bAutoPlay)
 	{
@@ -197,6 +200,20 @@ void FMediaPlateTrackEditor::AddTrackForComponent(UMediaPlateComponent* Componen
 				break;
 		}
 	}
+
+	// Media Track doesn't support audio from the media player yet.
+	if (Component->bEnableAudio)
+	{
+		// Ask the user if it can be disabled.
+		const EAppReturnType::Type Reply = FMessageDialog::Open(EAppMsgType::YesNo,
+			LOCTEXT("Prompt_DisableAudio", "Would you like to disable audio on the MediaPlate actor (recommended)?"));
+
+		if (Reply == EAppReturnType::Yes)
+		{
+			Component->Modify();
+			Component->SetEnableAudio(false);
+		}
+	}
 }
 
 void FMediaPlateTrackEditor::OnRelease()
