[MediaPlate] Disable Media Plate's control while it is controlled by sequencer.

- MediaPlateComponent: adding an accurate mechanism to keep track of the media textures in use by sequencer. The texture layers and "normal mode" are not accurate for this purpose and I don't want to change it to avoid breaking other features using it. It is private so we can refactor it later if needed. There will be another refactor for proxying the audio component (eventually).
- MediaPlateComponent: Closing the internal player as soon as a media texture is acquired for proxying. It could be argued that it would only need to close the corresponding player, but the audio issue would remain in that case. The purpose of this is to stop audio from the main player's audio track. Therefore, as soon as a texture is used, meaning sequencer is evaluation a section bound to that media plate, even if it is not the "main" one, it should still close the internal player to stop the audio. This constraint will likely be revisited when audio is proxied along with the texture. This would allow secondary texture/audio component (from sequencer sections that are not media source proxies) to be proxied while main player (which would not be proxied) could still be controlled.
- MediaPlateCustomization: disable the playback controls (Open and Play in particular) when the media plate is being proxied by sequencer. This is the main issue that was reported.
- MediaPlateCustomization: reworked the multi-selection handling so that the button's enabling are the union of the selection and not just the first one.

diff --git a/Engine/Plugins/Media/MediaPlate/Source/MediaPlate/Private/MediaPlateComponent.cpp b/Engine/Plugins/Media/MediaPlate/Source/MediaPlate/Private/MediaPlateComponent.cpp
--- a/Engine/Plugins/Media/MediaPlate/Source/MediaPlate/Private/MediaPlateComponent.cpp
+++ b/Engine/Plugins/Media/MediaPlate/Source/MediaPlate/Private/MediaPlateComponent.cpp
@@ -173,8 +173,8 @@ void UMediaPlateComponent::OnRegister()
 	MediaPlayer->OnPlaybackResumed.AddUniqueDynamic(this, &UMediaPlateComponent::OnMediaResumed);
 	MediaPlayer->OnPlaybackSuspended.AddUniqueDynamic(this, &UMediaPlateComponent::OnMediaSuspended);
 
-	// Set up media texture. Note: don't set player if in proxy mode.
-	SetUpTextures(IsNormalMode() ? UE::MediaPlateComponent::ESetUpTexturesFlags::AllowSetPlayer : UE::MediaPlateComponent::ESetUpTexturesFlags::None);
+	// Set up media texture. Note: will not set player to proxied textures.
+	SetUpTextures(UE::MediaPlateComponent::ESetUpTexturesFlags::AllowSetPlayer);
 
 	// Set up sound component if we have one.
 	if (SoundComponent != nullptr)
@@ -315,6 +315,19 @@ UMediaTexture* UMediaPlateComponent::GetMediaTexture(int32 Index)
 	return MediaTexture;
 }
 
+bool UMediaPlateComponent::IsEventStateChangeAllowed(EMediaPlateEventState InEventState) const
+{
+	// Note: centralizing the state change restriction here to make it easier to maintain.
+	switch(InEventState)
+	{
+	case EMediaPlateEventState::Open:
+	case EMediaPlateEventState::Play:
+		return !IsAnyMediaTextureProxied();
+	default:
+		return true;
+	}
+}
+
 void UMediaPlateComponent::Open()
 {
 	bIsMediaPlatePlaying = true;
@@ -795,6 +808,15 @@ UMediaTexture* UMediaPlateComponent::ProxyGetMediaTexture(int32 LayerIndex, int3
 	UMediaTexture* MediaTexture = GetMediaTexture(TextureIndex);
 	if (MediaTexture != nullptr)
 	{
+		ProxiedMediaTextureIndices.Add(TextureIndex);
+
+		// Ensure internal media player is closed to not interfere with external one (audio in particular).
+		// todo: revisit this when audio proxying is supported. It might be possible to be more permissive.
+		if (IsMediaPlatePlaying())
+		{
+			Close();
+		}
+		
 		SetNormalMode(false);
 		if (TextureLayers.Num() < LayerIndex + 1)
 		{
@@ -825,6 +847,7 @@ UMediaTexture* UMediaPlateComponent::ProxyGetMediaTexture(int32 LayerIndex, int3
 
 void UMediaPlateComponent::ProxyReleaseMediaTexture(int32 LayerIndex, int32 TextureIndex)
 {
+	ProxiedMediaTextureIndices.Remove(TextureIndex);
 	ProxySetTextureBlend(LayerIndex, TextureIndex, 0.0f);
 
 	if (LayerIndex < TextureLayers.Num())
@@ -1161,9 +1184,9 @@ void UMediaPlateComponent::SetUpTextures(UE::MediaPlateComponent::ESetUpTextures
 {
 	// Prevent media texture blackouts by only updating resource and material uniforms on relevant changes.
 	bool bApplyMaterialUpdate = false;
-	for (UMediaTexture* MediaTexture : MediaTextures)
+	for (int32 MediaTextureIndex = 0; MediaTextureIndex < MediaTextures.Num(); ++MediaTextureIndex)
 	{
-		if (MediaTexture != nullptr)
+		if (UMediaTexture* MediaTexture = MediaTextures[MediaTextureIndex])
 		{
 			bool bApplyTextureUpdate = false;
 
@@ -1180,8 +1203,10 @@ void UMediaPlateComponent::SetUpTextures(UE::MediaPlateComponent::ESetUpTextures
 				bApplyMaterialUpdate = true;
 			}
 
-			if ((EnumHasAllFlags(Flags, UE::MediaPlateComponent::ESetUpTexturesFlags::AllowSetPlayer)) &&
-				(MediaTexture->GetMediaPlayer() != MediaPlayer.Get()))
+			// Even if setting the player is allowed, we don't set the player if the texture is currently being proxied.
+			if (EnumHasAllFlags(Flags, UE::MediaPlateComponent::ESetUpTexturesFlags::AllowSetPlayer)
+				&& MediaTexture->GetMediaPlayer() != MediaPlayer.Get()
+				&& !ProxiedMediaTextureIndices.Contains(MediaTextureIndex))
 			{
 				MediaTexture->SetMediaPlayer(MediaPlayer);
 				bApplyTextureUpdate = true;
@@ -1250,6 +1275,11 @@ bool UMediaPlateComponent::IsNormalMode() const
 #endif
 }
 
+bool UMediaPlateComponent::IsAnyMediaTextureProxied() const
+{
+	return !ProxiedMediaTextureIndices.IsEmpty();
+}
+
 void UMediaPlateComponent::UpdateTextureLayers()
 {
 	if (AMediaPlate* MediaPlate = GetOwner<AMediaPlate>())
diff --git a/Engine/Plugins/Media/MediaPlate/Source/MediaPlate/Public/MediaPlateComponent.h b/Engine/Plugins/Media/MediaPlate/Source/MediaPlate/Public/MediaPlateComponent.h
--- a/Engine/Plugins/Media/MediaPlate/Source/MediaPlate/Public/MediaPlateComponent.h
+++ b/Engine/Plugins/Media/MediaPlate/Source/MediaPlate/Public/MediaPlateComponent.h
@@ -105,6 +105,15 @@ public:
 	UFUNCTION(BlueprintCallable, Category = "Media|MediaPlateComponent")
 	UMediaTexture* GetMediaTexture(int32 Index = 0);
 
+	/**
+	 * Indicates if switching to the given event state (open, play, etc) is currently allowed by
+	 * the backend.
+	 * @param InRequestEventState Requested event state to switch to (Open, Play, etc)
+	 * @return true if the state switch is allowed, false otherwise.
+	 */
+	UFUNCTION(BlueprintCallable, Category = "Media|MediaPlateComponent")
+	bool IsEventStateChangeAllowed(EMediaPlateEventState InRequestEventState) const;
+
 	/**
 	 * Call this to open the media.
 	 */
@@ -480,6 +489,12 @@ private:
 	};
 	TArray<Layer> TextureLayers;
 
+	/**
+	 * This accurately reflects the set of textures that are actively being used externally
+	 * through the IMediaPlayerProxyInterface's Get/ReleaseMediaTexture functions.
+	 */
+	TSet<int32> ProxiedMediaTextureIndices;
+
 	/**
 	 * Plays a media source.
 	 *
@@ -578,7 +593,13 @@ private:
 	 * Returns true if normal mode of false if proxy mode.
 	 */
 	bool IsNormalMode() const;
-	
+
+	/**
+	 * @brief Indicates if any media texture is currently being proxied by sequencer.
+	 * @return true if at least one media texture is in use by sequencer.
+	 */
+	bool IsAnyMediaTextureProxied() const;
+
 	/**
 	 * Sets textures in our material according to the layer assignments.
 	 */
diff --git a/Engine/Plugins/Media/MediaPlate/Source/MediaPlateEditor/Private/MediaPlateCustomization.cpp b/Engine/Plugins/Media/MediaPlate/Source/MediaPlateEditor/Private/MediaPlateCustomization.cpp
--- a/Engine/Plugins/Media/MediaPlate/Source/MediaPlateEditor/Private/MediaPlateCustomization.cpp
+++ b/Engine/Plugins/Media/MediaPlate/Source/MediaPlateEditor/Private/MediaPlateCustomization.cpp
@@ -118,23 +118,15 @@ void FMediaPlateCustomization::CustomizeDetails(IDetailLayoutBuilder& DetailBuil
 						.VAlign(VAlign_Center)
 						.IsEnabled_Lambda([WeakSelf]
 						{
-							const TSharedPtr<FMediaPlateCustomization> Self = WeakSelf.Pin();
-							if (Self.IsValid())
+							if (const TSharedPtr<FMediaPlateCustomization> Self = WeakSelf.Pin())
 							{
-								if (const UMediaPlayer* MediaPlayer = Self->GetMediaPlayer())
-								{
-									return MediaPlayer->IsReady() &&
-										MediaPlayer->SupportsSeeking() &&
-										MediaPlayer->GetTime() > FTimespan::Zero();
-								}
+								return Self->IsButtonEventAllowedForAnyPlate(EMediaPlateEventState::Rewind);
 							}
-
 							return false;
 						})
 						.OnClicked_Lambda([WeakSelf]() -> FReply
 						{
-							const TSharedPtr<FMediaPlateCustomization> Self = WeakSelf.Pin();
-							if (Self.IsValid())
+							if (const TSharedPtr<FMediaPlateCustomization> Self = WeakSelf.Pin())
 							{
 								Self->OnButtonEvent(EMediaPlateEventState::Rewind);
 								return FReply::Handled();
@@ -158,21 +150,15 @@ void FMediaPlateCustomization::CustomizeDetails(IDetailLayoutBuilder& DetailBuil
 						.VAlign(VAlign_Center)
 						.IsEnabled_Lambda([WeakSelf]
 						{
-							const TSharedPtr<FMediaPlateCustomization> Self = WeakSelf.Pin();
-							if (Self.IsValid())
+							if (const TSharedPtr<FMediaPlateCustomization> Self = WeakSelf.Pin())
 							{
-								if (UMediaPlayer* MediaPlayer = Self->GetMediaPlayer())
-								{
-									return MediaPlayer->IsReady() && MediaPlayer->SupportsRate(UMediaPlateComponent::GetReverseRate(MediaPlayer), false);
-								}
+								return Self->IsButtonEventAllowedForAnyPlate(EMediaPlateEventState::Reverse);
 							}
-
 							return false;
 						})
 						.OnClicked_Lambda([WeakSelf]() -> FReply
 						{
-							const TSharedPtr<FMediaPlateCustomization> Self = WeakSelf.Pin();
-							if (Self.IsValid())
+							if (const TSharedPtr<FMediaPlateCustomization> Self = WeakSelf.Pin())
 							{
 								Self->OnButtonEvent(EMediaPlateEventState::Reverse);
 								return FReply::Handled();
@@ -196,36 +182,29 @@ void FMediaPlateCustomization::CustomizeDetails(IDetailLayoutBuilder& DetailBuil
 						.VAlign(VAlign_Center)
 						.IsEnabled_Lambda([WeakSelf]
 						{
-							const TSharedPtr<FMediaPlateCustomization> Self = WeakSelf.Pin();
-							if (Self.IsValid())
+							if (const TSharedPtr<FMediaPlateCustomization> Self = WeakSelf.Pin())
 							{
-								// Is the player paused or fast forwarding/rewinding?
-								if (const UMediaPlayer* MediaPlayer = Self->GetMediaPlayer())
-								{
-									return MediaPlayer->IsReady()
-										&& (!MediaPlayer->IsPlaying() || (MediaPlayer->GetRate() != 1.0f));
-								}
+								return Self->IsButtonEventAllowedForAnyPlate(EMediaPlateEventState::Play);
 							}
 
 							return false;
 						})
 						.Visibility_Lambda([WeakSelf]
 						{
-							const TSharedPtr<FMediaPlateCustomization> Self = WeakSelf.Pin();
-							if (Self.IsValid())
+							if (const TSharedPtr<FMediaPlateCustomization> Self = WeakSelf.Pin())
 							{
-								if (const UMediaPlayer* MediaPlayer = Self->GetMediaPlayer())
+								const bool bAllPlaying = Self->IsTrueForAllPlayers([](const UMediaPlayer* InMediaPlayer)
 								{
-									return MediaPlayer->IsPlaying() ? EVisibility::Collapsed : EVisibility::Visible;
-								}
+									return InMediaPlayer->IsPlaying();
+								});
+								return bAllPlaying ? EVisibility::Collapsed : EVisibility::Visible;
 							}
 
 							return EVisibility::Visible;
 						})
 						.OnClicked_Lambda([WeakSelf]() -> FReply
 						{
-							const TSharedPtr<FMediaPlateCustomization> Self = WeakSelf.Pin();
-							if (Self.IsValid())
+							if (const TSharedPtr<FMediaPlateCustomization> Self = WeakSelf.Pin())
 							{
 								Self->OnButtonEvent(EMediaPlateEventState::Play);
 								return FReply::Handled();
@@ -250,35 +229,27 @@ void FMediaPlateCustomization::CustomizeDetails(IDetailLayoutBuilder& DetailBuil
 						.VAlign(VAlign_Center)
 						.IsEnabled_Lambda([WeakSelf]
 						{
-							const TSharedPtr<FMediaPlateCustomization> Self = WeakSelf.Pin();
-							if (Self.IsValid())
+							if (const TSharedPtr<FMediaPlateCustomization> Self = WeakSelf.Pin())
 							{
-								if (const UMediaPlayer* MediaPlayer = Self->GetMediaPlayer())
-								{
-									return MediaPlayer->CanPause() && !MediaPlayer->IsPaused();
-								}
+								return Self->IsButtonEventAllowedForAnyPlate(EMediaPlateEventState::Pause);
 							}
-
 							return false;
 						})
 						.Visibility_Lambda([WeakSelf]
 						{
-							const TSharedPtr<FMediaPlateCustomization> Self = WeakSelf.Pin();
-							if (Self.IsValid())
+							if (const TSharedPtr<FMediaPlateCustomization> Self = WeakSelf.Pin())
 							{
-								if (const UMediaPlayer* MediaPlayer = Self->GetMediaPlayer())
+								const bool bAllPaused = Self->IsTrueForAllPlayers([](const UMediaPlayer* InMediaPlayer)
 								{
-									const bool bIsVisible = MediaPlayer->CanPause() && !MediaPlayer->IsPaused();
-									return bIsVisible ? EVisibility::Visible : EVisibility::Collapsed;
-								}
+									return InMediaPlayer->IsPaused();
+								});
+								return bAllPaused ? EVisibility::Collapsed : EVisibility::Visible;
 							}
-
 							return EVisibility::Collapsed;
 						})
 						.OnClicked_Lambda([WeakSelf]() -> FReply
 						{
-							const TSharedPtr<FMediaPlateCustomization> Self = WeakSelf.Pin();
-							if (Self.IsValid())
+							if (const TSharedPtr<FMediaPlateCustomization> Self = WeakSelf.Pin())
 							{
 								Self->OnButtonEvent(EMediaPlateEventState::Pause);
 								return FReply::Handled();
@@ -302,21 +273,15 @@ void FMediaPlateCustomization::CustomizeDetails(IDetailLayoutBuilder& DetailBuil
 					.VAlign(VAlign_Center)
 					.IsEnabled_Lambda([WeakSelf]
 					{
-						const TSharedPtr<FMediaPlateCustomization> Self = WeakSelf.Pin();
-						if (Self.IsValid())
+						if (const TSharedPtr<FMediaPlateCustomization> Self = WeakSelf.Pin())
 						{
-							if (UMediaPlayer* MediaPlayer = Self->GetMediaPlayer())
-							{
-								return MediaPlayer->IsReady() && MediaPlayer->SupportsRate(UMediaPlateComponent::GetForwardRate(MediaPlayer), false);
-							}
+							return Self->IsButtonEventAllowedForAnyPlate(EMediaPlateEventState::Forward);
 						}
-
 						return false;
 					})
 					.OnClicked_Lambda([WeakSelf]() -> FReply
 					{
-						const TSharedPtr<FMediaPlateCustomization> Self = WeakSelf.Pin();
-						if (Self.IsValid())
+						if (const TSharedPtr<FMediaPlateCustomization> Self = WeakSelf.Pin())
 						{
 							Self->OnButtonEvent(EMediaPlateEventState::Forward);
 							return FReply::Handled();
@@ -338,10 +303,17 @@ void FMediaPlateCustomization::CustomizeDetails(IDetailLayoutBuilder& DetailBuil
 				[
 					SNew(SButton)
 						.VAlign(VAlign_Center)
+						.IsEnabled_Lambda([WeakSelf]
+						{
+							if (const TSharedPtr<FMediaPlateCustomization> Self = WeakSelf.Pin())
+							{
+								return Self->IsButtonEventAllowedForAnyPlate(EMediaPlateEventState::Open);
+							}
+							return true;
+						})
 						.OnClicked_Lambda([WeakSelf]() -> FReply
 						{
-							const TSharedPtr<FMediaPlateCustomization> Self = WeakSelf.Pin();
-							if (Self.IsValid())
+							if (const TSharedPtr<FMediaPlateCustomization> Self = WeakSelf.Pin())
 							{
 								Self->OnButtonEvent(EMediaPlateEventState::Open);
 								return FReply::Handled();
@@ -365,21 +337,15 @@ void FMediaPlateCustomization::CustomizeDetails(IDetailLayoutBuilder& DetailBuil
 						.VAlign(VAlign_Center)
 						.IsEnabled_Lambda([WeakSelf]
 						{
-							const TSharedPtr<FMediaPlateCustomization> Self = WeakSelf.Pin();
-							if (Self.IsValid())
+							if (const TSharedPtr<FMediaPlateCustomization> Self = WeakSelf.Pin())
 							{
-								if (UMediaPlayer* MediaPlayer = Self->GetMediaPlayer())
-								{
-									return !MediaPlayer->GetUrl().IsEmpty();
-								}
+								return Self->IsButtonEventAllowedForAnyPlate(EMediaPlateEventState::Close);
 							}
-
 							return false;
 						})
 						.OnClicked_Lambda([WeakSelf]() -> FReply
 						{
-							const TSharedPtr<FMediaPlateCustomization> Self = WeakSelf.Pin();
-							if (Self.IsValid())
+							if (const TSharedPtr<FMediaPlateCustomization> Self = WeakSelf.Pin())
 							{
 								Self->OnButtonEvent(EMediaPlateEventState::Close);
 								return FReply::Handled();
@@ -940,33 +906,94 @@ void FMediaPlateCustomization::OnStaticMeshChanged(const FAssetData& AssetData)
 	}
 }
 
-void FMediaPlateCustomization::OnButtonEvent(EMediaPlateEventState State)
+bool FMediaPlateCustomization::IsButtonEventAllowedForAnyPlate(EMediaPlateEventState InState) const
+{
+	// Returns true if any of the selected plate allows the action.
+	for (const TWeakObjectPtr<UMediaPlateComponent>& MediaPlateWeak : MediaPlatesList)
+	{
+		if (UMediaPlateComponent* MediaPlate = MediaPlateWeak.Get())
+		{
+			if (MediaPlate->IsEventStateChangeAllowed(InState)
+				&& IsButtonEventAllowedForPlayer(InState, MediaPlate->GetMediaPlayer()))
+			{
+				return true;
+			}
+		}
+	}
+	return false;
+}
+
+bool FMediaPlateCustomization::IsButtonEventAllowedForPlayer(EMediaPlateEventState InState, UMediaPlayer* InMediaPlayer) const
+{
+	// Note: centralize the state switch conditions here to make it easier to maintain.
+	if (!InMediaPlayer)
+	{
+		return false;
+	}
+	
+	switch (InState)
+	{
+	case EMediaPlateEventState::Play:
+		// Is player paused or fast forwarding/rewinding?
+		return InMediaPlayer->IsReady()
+			&& (!InMediaPlayer->IsPlaying() || InMediaPlayer->GetRate() != 1.0f);
+
+	case EMediaPlateEventState::Open:
+		return true; // The condition is implemented by the media plate already.
+
+	case EMediaPlateEventState::Close:
+		return !InMediaPlayer->GetUrl().IsEmpty();
+
+	case EMediaPlateEventState::Pause:
+		return InMediaPlayer->CanPause() && !InMediaPlayer->IsPaused();
+
+	case EMediaPlateEventState::Reverse:
+		return InMediaPlayer->IsReady() && InMediaPlayer->SupportsRate(UMediaPlateComponent::GetReverseRate(InMediaPlayer), false);
+
+	case EMediaPlateEventState::Forward:
+		return InMediaPlayer->IsReady() && InMediaPlayer->SupportsRate(UMediaPlateComponent::GetForwardRate(InMediaPlayer), false);
+
+	case EMediaPlateEventState::Rewind:
+		return InMediaPlayer->IsReady() && InMediaPlayer->SupportsSeeking() && InMediaPlayer->GetTime() > FTimespan::Zero();
+
+	//case EMediaPlateEventState::Next:	// was not implemented
+	//case EMediaPlateEventState::Previous: // was not implemented
+
+	default:
+		return true;
+	}
+}
+
+void FMediaPlateCustomization::OnButtonEvent(EMediaPlateEventState InState)
 {
 	IMediaAssetsModule* MediaAssets = FModuleManager::LoadModulePtr<IMediaAssetsModule>("MediaAssets");
 	TArray<FString> ActorsPathNames;
 	ActorsPathNames.Reserve(MediaPlatesList.Num());
-	for (TWeakObjectPtr<UMediaPlateComponent>& MediaPlatePtr : MediaPlatesList)
+	for (const TWeakObjectPtr<UMediaPlateComponent>& MediaPlateWeak : MediaPlatesList)
 	{
-		if (!MediaPlatePtr.IsValid())
+		UMediaPlateComponent* MediaPlate = MediaPlateWeak.Get();
+
+		// Note: because of multi-selection and the possibility of different player states, we need check restrictions per plate again.
+		if (!MediaPlate || !MediaPlate->IsEventStateChangeAllowed(InState) || !IsButtonEventAllowedForPlayer(InState, MediaPlate->GetMediaPlayer()))
 		{
 			continue;
 		}
 
-		ActorsPathNames.Add(MediaPlatePtr->GetOwner()->GetPathName());
+		ActorsPathNames.Add(MediaPlate->GetOwner()->GetPathName());
 
-		if (State == EMediaPlateEventState::Open)
+		if (InState == EMediaPlateEventState::Open)
 		{
 			// Tell the editor module that this media plate is playing.
 			FMediaPlateEditorModule* EditorModule = FModuleManager::LoadModulePtr<FMediaPlateEditorModule>("MediaPlateEditor");
 			if (EditorModule != nullptr)
 			{
-				EditorModule->MediaPlateStartedPlayback(MediaPlatePtr.Get());
+				EditorModule->MediaPlateStartedPlayback(MediaPlate);
 			}
 		}
 
-		MediaPlatePtr->SwitchStates(State);
+		MediaPlate->SwitchStates(InState);
 	}
-	MediaAssets->BroadcastOnMediaStateChangedEvent(ActorsPathNames, (uint8)State);
+	MediaAssets->BroadcastOnMediaStateChangedEvent(ActorsPathNames, (uint8)InState);
 }
 
 FReply FMediaPlateCustomization::OnOpenMediaPlate()
@@ -996,38 +1023,44 @@ void FMediaPlateCustomization::StopMediaPlates()
 	OnButtonEvent(EMediaPlateEventState::Close);
 }
 
-UMediaPlayer* FMediaPlateCustomization::GetMediaPlayer() const
+TArray<TWeakObjectPtr<UMediaPlayer>> FMediaPlateCustomization::GetMediaPlayers() const
 {
-	for (const TWeakObjectPtr<UMediaPlateComponent>& MediaPlatePtr : MediaPlatesList)
+	TArray<TWeakObjectPtr<UMediaPlayer>> MediaPlayers;
+	MediaPlayers.Reserve(MediaPlatesList.Num());
+	
+	for (const TWeakObjectPtr<UMediaPlateComponent>& MediaPlateWeak : MediaPlatesList)
 	{
-		if (UMediaPlateComponent* MediaPlate = MediaPlatePtr.Get())
+		if (UMediaPlateComponent* MediaPlate = MediaPlateWeak.Get())
 		{
 			if (UMediaPlayer* MediaPlayer = MediaPlate->GetMediaPlayer())
 			{
-				return MediaPlayer;
+				MediaPlayers.Add(MediaPlayer);
 			}
 		}
 	}
-
-	return nullptr;
+	return MediaPlayers;
 }
 
-TArray<TWeakObjectPtr<UMediaPlayer>> FMediaPlateCustomization::GetMediaPlayers() const
+bool FMediaPlateCustomization::IsTrueForAllPlayers(TFunctionRef<bool(const UMediaPlayer*)> InPredicate) const
 {
-	TArray<TWeakObjectPtr<UMediaPlayer>> MediaPlayers;
-	MediaPlayers.Reserve(MediaPlatesList.Num());
-	
+	bool bPredicateCalled = false;
+
 	for (const TWeakObjectPtr<UMediaPlateComponent>& MediaPlateWeak : MediaPlatesList)
 	{
 		if (UMediaPlateComponent* MediaPlate = MediaPlateWeak.Get())
 		{
-			if (UMediaPlayer* MediaPlayer = MediaPlate->GetMediaPlayer())
+			if (const UMediaPlayer* MediaPlayer = MediaPlate->GetMediaPlayer())
 			{
-				MediaPlayers.Add(MediaPlayer);
+				if (!InPredicate(MediaPlayer))
+				{
+					return false;
+				}
+				bPredicateCalled = true;
 			}
 		}
 	}
-	return MediaPlayers;
+
+	return bPredicateCalled;
 }
 
 void FMediaPlateCustomization::CustomizeCategories(IDetailLayoutBuilder& InDetailBuilder)
diff --git a/Engine/Plugins/Media/MediaPlate/Source/MediaPlateEditor/Private/MediaPlateCustomization.h b/Engine/Plugins/Media/MediaPlate/Source/MediaPlateEditor/Private/MediaPlateCustomization.h
--- a/Engine/Plugins/Media/MediaPlate/Source/MediaPlateEditor/Private/MediaPlateCustomization.h
+++ b/Engine/Plugins/Media/MediaPlate/Source/MediaPlateEditor/Private/MediaPlateCustomization.h
@@ -170,10 +170,21 @@ private:
 	 */
 	void OnStaticMeshChanged(const FAssetData& AssetData);
 
+	/**
+	 * Determines if the event state switch is allowed in at least one selected components.
+	 */
+	bool IsButtonEventAllowedForAnyPlate(EMediaPlateEventState InState) const;
+
+	/**
+	 * Determines if the event state switch is allowed for the given media player.
+	 * Those are the UI only conditions to mirror the backend's conditions.
+	 */
+	bool IsButtonEventAllowedForPlayer(EMediaPlateEventState InState, UMediaPlayer* InMediaPlayer) const;
+
 	/**
 	* Changes the state of selected media plates and broadcasts the even to the remote endpoints.
 	*/
-	void OnButtonEvent(EMediaPlateEventState State);
+	void OnButtonEvent(EMediaPlateEventState InState);
 
 	/**
 	 * Called when the open media plate button is pressed.
@@ -185,7 +196,10 @@ private:
 	 */
 	void StopMediaPlates();
 
-	UMediaPlayer* GetMediaPlayer() const;
-
 	TArray<TWeakObjectPtr<UMediaPlayer>> GetMediaPlayers() const;
+
+	/**
+	 * Evaluates if the given predicate is true for all players.
+	 */
+	bool IsTrueForAllPlayers(TFunctionRef<bool(const UMediaPlayer*)> InPredicate) const;
 };
