[MediaPlate] Media Plate support for transition logic, dynamic level instancing and remote control of the media source.

- Adding latent open for support of a transition logic blueprint task.
- Close the player in EndPlay(). This is need for using a MediaPlate in a dynamic level instance. Without it, the player keeps going and consume resources.
- Fixing RefreshMediaPlateResource so it doesn't close the player if called multiple times without actual change (RC does that).

diff --git a/Engine/Plugins/Media/MediaPlate/Source/MediaPlate/Private/MediaPlateComponent.cpp b/Engine/Plugins/Media/MediaPlate/Source/MediaPlate/Private/MediaPlateComponent.cpp
--- a/Engine/Plugins/Media/MediaPlate/Source/MediaPlate/Private/MediaPlateComponent.cpp
+++ b/Engine/Plugins/Media/MediaPlate/Source/MediaPlate/Private/MediaPlateComponent.cpp
@@ -3,6 +3,7 @@
 #include "MediaPlateComponent.h"
 
 #include "Components/StaticMeshComponent.h"
+#include "Engine/Engine.h"
 #include "Engine/StaticMesh.h"
 #include "IMediaAssetsModule.h"
 #include "IMediaClock.h"
@@ -12,6 +13,7 @@
 #include "Materials/MaterialInstanceDynamic.h"
 #include "MediaPlate.h"
 #include "MediaPlateModule.h"
+#include "MediaPlateOpenLatentAction.h"
 #include "MediaPlayer.h"
 #include "MediaPlaylist.h"
 #include "MediaSoundComponent.h"
@@ -65,6 +67,11 @@ namespace UE::MediaPlateComponent
 		}
 #endif
 	}
+
+	UMediaSource* GetMediaSource(UMediaPlaylist* InPlaylist, int32 InPlaylistIndex)
+	{
+		return InPlaylist ? InPlaylist->Get(InPlaylistIndex) : nullptr;
+	}
 };
 
 /**
@@ -199,6 +206,14 @@ void UMediaPlateComponent::BeginPlay()
 	}
 }
 
+void UMediaPlateComponent::EndPlay(const EEndPlayReason::Type InEndPlayReason)
+{
+	// Ensure media resources are properly released.
+	Close();
+
+	Super::EndPlay(InEndPlayReason);
+}
+
 void UMediaPlateComponent::BeginDestroy()
 {
 	if (ClockSink.IsValid())
@@ -345,11 +360,7 @@ void UMediaPlateComponent::Open()
 		bool bIsPlaying = false;
 		if (MediaPlayer != nullptr)
 		{
-			UMediaSource* MediaSource = nullptr;
-			if (UMediaPlaylist* Playlist = GetMediaPlaylist())
-			{
-				MediaSource = Playlist->Get(PlaylistIndex);
-			}
+			UMediaSource* MediaSource = UE::MediaPlateComponent::GetMediaSource(GetMediaPlaylist(), PlaylistIndex);
 			bIsPlaying = PlayMediaSource(MediaSource, bPlayOnOpen);
 		}
 
@@ -368,6 +379,21 @@ void UMediaPlateComponent::Open()
 	UpdateTicking();
 }
 
+void UMediaPlateComponent::OpenLatent(const UObject* InWorldContextObject, struct FLatentActionInfo InLatentInfo, float InTimeout, bool bInWaitForTexture, bool& bOutSuccess)
+{
+	bOutSuccess = false;
+
+	if (UWorld* World = GEngine->GetWorldFromContextObject(InWorldContextObject, EGetWorldErrorMode::LogAndReturnNull))
+	{
+		FLatentActionManager& LatentManager = World->GetLatentActionManager();
+		if (LatentManager.FindExistingAction<FMediaPlateOpenLatentAction>(InLatentInfo.CallbackTarget, InLatentInfo.UUID) == nullptr)
+		{
+			FMediaPlateOpenLatentAction* NewAction = new FMediaPlateOpenLatentAction(InLatentInfo, this, InTimeout, bInWaitForTexture, bOutSuccess);
+			LatentManager.AddNewAction(InLatentInfo.CallbackTarget, InLatentInfo.UUID, NewAction);
+		}
+	}
+}
+
 bool UMediaPlateComponent::Next()
 {
 	UMediaPlaylist* Playlist = GetMediaPlaylist();
@@ -794,13 +820,7 @@ const FMediaSourceCacheSettings& UMediaPlateComponent::GetCacheSettings() const
 
 UMediaSource* UMediaPlateComponent::ProxyGetMediaSourceFromIndex(int32 Index) const
 {
-	UMediaSource* MediaSource = nullptr;
-
-	if (UMediaPlaylist* Playlist = GetMediaPlaylist())
-	{
-		MediaSource = Playlist->Get(Index);
-	}
-	return MediaSource;
+	return UE::MediaPlateComponent::GetMediaSource(GetMediaPlaylist(), Index);
 }
 
 UMediaTexture* UMediaPlateComponent::ProxyGetMediaTexture(int32 LayerIndex, int32 TextureIndex)
@@ -1412,6 +1432,19 @@ void UMediaPlateComponent::RefreshMediaPlateResource()
 {
 	MediaPlateResource.RefreshActivePlaylist(this);
 
+	// Remote Control may call PostEditChangeProperty more than once from various paths
+	// with some delays causing Media Plate to close when it shouldn't. We mitigate that by
+	// skipping the closing/restarting logic if the media source didn't change.
+	if (MediaPlayer)
+	{
+		const UMediaSource* MediaSource = UE::MediaPlateComponent::GetMediaSource(GetMediaPlaylist(), PlaylistIndex);
+		const UMediaSource* PlayerMediaSource = UE::MediaPlateComponent::GetMediaSource(MediaPlayer->GetPlaylist(), MediaPlayer->GetPlaylistIndex());
+		if (MediaSource == PlayerMediaSource)
+		{
+			return;
+		}
+	}
+
 	if (RestartPlayer())
 	{
 		return;
diff --git a/Engine/Plugins/Media/MediaPlate/Source/MediaPlate/Private/MediaPlateOpenLatentAction.cpp b/Engine/Plugins/Media/MediaPlate/Source/MediaPlate/Private/MediaPlateOpenLatentAction.cpp
new file mode 100644
--- /dev/null
+++ b/Engine/Plugins/Media/MediaPlate/Source/MediaPlate/Private/MediaPlateOpenLatentAction.cpp
@@ -0,0 +1,300 @@
+// Copyright Epic Games, Inc. All Rights Reserved.
+
+#include "MediaPlateOpenLatentAction.h"
+
+#include "IMediaEventSink.h"
+#include "MediaHelpers.h"
+#include "MediaPlateComponent.h"
+#include "MediaPlateModule.h"
+#include "MediaPlayer.h"
+#include "MediaTexture.h"
+
+// Implementation references:
+// - UMediaSourceRenderer: "wait on texture" part.
+// - UMediaPlayer::OpenSourceLatent: player state management.
+
+FMediaPlateOpenLatentAction::FMediaPlateOpenLatentAction(const FLatentActionInfo& InLatentInfo,
+	UMediaPlateComponent* InMediaPlateComponent, float InWaitTimeout, bool bInWaitForTexture, bool& InOutSuccess)
+	: ExecutionFunction(InLatentInfo.ExecutionFunction)
+	, OutputLink(InLatentInfo.Linkage)
+	, CallbackTarget(InLatentInfo.CallbackTarget)
+	, MediaPlateComponentWeak(InMediaPlateComponent)
+	, bWaitForTexture(bInWaitForTexture)
+	, OutSuccess(InOutSuccess)
+	, TimeRemaining(InWaitTimeout)
+{
+	if (InMediaPlateComponent)
+	{
+		InMediaPlateComponent->Open();
+
+		if (UMediaPlayer* MediaPlayer = InMediaPlateComponent->GetMediaPlayer())
+		{
+			MediaPlayerWeak = MediaPlayer;
+			MediaPlayer->OnMediaEvent().AddRaw(this, &FMediaPlateOpenLatentAction::OnMediaEvent);
+			URL = MediaPlayer->GetUrl();
+
+			if (bWaitForTexture)
+			{
+				// Reset the aspect ratio to be able to detect when a sample has been processed by the texture.
+				if (UMediaTexture* Texture = InMediaPlateComponent->GetMediaTexture())
+				{
+					Texture->CurrentAspectRatio = 0.0f;
+				}
+			}
+		}
+		else
+		{
+			UE_LOG(LogMediaPlate, Warning, TEXT("[%d] Media Plate Open Latent: Failed initial open: %s"), OutputLink, *URL);
+			bSawError = true;
+		}
+	}
+	else
+	{
+		UE_LOG(LogMediaPlate, Warning, TEXT("[%d] Media Plate Open Latent: Failed initial open because no media source given"), OutputLink);
+		bSawError = true;
+	}
+}
+
+FMediaPlateOpenLatentAction::~FMediaPlateOpenLatentAction()
+{
+	if (UMediaPlayer* MediaPlayer = MediaPlayerWeak.Get())
+	{
+		MediaPlayer->OnMediaEvent().RemoveAll(this);	
+	}
+}
+
+void FMediaPlateOpenLatentAction::UpdateOperation(FLatentResponse& InResponse)
+{
+	if (bSawMediaOpenFailed)
+	{
+		UE_LOG(LogMediaPlate, Warning, TEXT("[%d] Media Plate Open Latent: Saw media open failed event. %s"), OutputLink, *URL);
+		FailedOperation(InResponse);
+		return;
+	}
+
+	UMediaPlateComponent* MediaPlateComponent = MediaPlateComponentWeak.Get();
+
+	if (!MediaPlateComponent)
+	{
+		UE_LOG(LogMediaPlate, Warning, TEXT("[%d] Media Plate Open Latent: Media Plate Component object was deleted. %s"), OutputLink, *URL);
+		FailedOperation(InResponse);
+		return;
+	}
+
+	// Protect against internal media player being deleted or swapped out.
+	const UMediaPlayer* CurrentMediaPlayer = MediaPlateComponent->GetMediaPlayer();
+	UMediaPlayer* MediaPlayer = MediaPlayerWeak.Get();
+	
+	if (!MediaPlayer || !CurrentMediaPlayer || MediaPlayer != CurrentMediaPlayer)
+	{
+		UE_LOG(LogMediaPlate, Warning, TEXT("[%d] Media Plate Open Latent: Media player object was deleted. %s"), OutputLink, *URL);
+		FailedOperation(InResponse);
+		return;
+	}
+
+	if (bSawError || MediaPlayer->HasError())
+	{
+		UE_LOG(LogMediaPlate, Warning, TEXT("[%d] Media Plate Open Latent: Media player is in Error state. %s"), OutputLink, *URL);
+		FailedOperation(InResponse);
+		return;
+	}
+
+	if (MediaPlayer->IsClosed() || bSawMediaClosed)
+	{
+		UE_LOG(LogMediaPlate, Warning, TEXT("[%d] Media Plate Open Latent: Media player is closed. %s"), OutputLink, *URL);
+		FailedOperation(InResponse);
+		return;
+	}
+
+	if (MediaPlayer->IsPreparing())
+	{
+		UE_LOG(LogMediaPlate, Verbose, TEXT("[%d] Media Plate Open Latent: Is preparing ... %s (Time out in %f seconds)"), OutputLink, *URL,  TimeRemaining);
+	}
+	else if (MediaPlayer->IsReady())
+	{
+		if (!bSawIsReady)
+		{
+			// Show this only once when the state is reached.
+			UE_LOG(LogMediaPlate, Verbose, TEXT("[%d] Media Plate Open Latent: IsReady() ... %s"), OutputLink, *URL);
+			bSawIsReady = true;
+		}
+
+		if (UpdateOperationPlayerReady(MediaPlayer, InResponse) == OperationUpdateResult::Completed)
+		{
+			return;
+		}
+	}
+	else
+	{
+		UE_LOG(LogMediaPlate, Verbose, TEXT("[%d] Media Plate Open Latent: Waiting for IsReady() ... %s (Time out in %f seconds)"), OutputLink, *URL, TimeRemaining);
+	}
+
+	// Update Timed out
+	TimeRemaining -= InResponse.ElapsedTime();
+	if (TimeRemaining <= 0.0f)
+	{
+		UE_LOG(LogMediaPlate, Warning, TEXT("[%d] Media Plate Open Latent: Timed out. %s"), OutputLink, *URL);
+		FailedOperation(InResponse);
+		return;
+	}
+}
+
+#if WITH_EDITOR
+FString FMediaPlateOpenLatentAction::GetDescription() const
+{
+	return FString::Printf(TEXT("Media Plate Open Latent: %s %s (Time out in %f seconds)"), *GetStatusString(), *URL, TimeRemaining);
+}
+#endif
+
+FString FMediaPlateOpenLatentAction::GetStatusString() const
+{
+	if (bSawMediaOpenFailed)
+	{
+		return TEXT("Media open failed event.");
+	}
+
+	UMediaPlateComponent* MediaPlateComponent = MediaPlateComponentWeak.Get();
+
+	if (!MediaPlateComponent)
+	{
+		return TEXT("Media Plate Component object was deleted.");
+	}
+
+	const UMediaPlayer* CurrentMediaPlayer = MediaPlateComponent->GetMediaPlayer();
+	const UMediaPlayer* MediaPlayer = MediaPlayerWeak.Get();
+
+	if (!MediaPlayer || !CurrentMediaPlayer || CurrentMediaPlayer != MediaPlayer)
+	{
+		return TEXT("Media player object was deleted.");
+	}
+
+	if (bSawError || MediaPlayer->HasError())
+	{
+		return TEXT("Media player is in Error state.");
+	}
+
+	if (MediaPlayer->IsClosed())
+	{
+		return TEXT("Media player is closed.");
+	}
+
+	if (MediaPlayer->IsPreparing())
+	{
+		return TEXT("Is preparing ...");
+	}
+
+	if (!MediaPlayer->IsReady())
+	{
+		return TEXT("Waiting for IsReady() ...");
+	}
+	
+	if (bWaitForTexture)
+	{
+		return TEXT("Is Ready - Waiting for Texture Render ...");
+	}
+
+	return FString::Printf(TEXT("Is Ready."));
+}
+
+FMediaPlateOpenLatentAction::OperationUpdateResult FMediaPlateOpenLatentAction::UpdateOperationPlayerReady(UMediaPlayer* InMediaPlayer, FLatentResponse& InResponse)
+{
+	if (bSawMediaOpened)
+	{
+		const FTimespan SeekTime = FTimespan::FromSeconds(MediaPlateComponentWeak->StartTime);
+
+		// We need to issue a seek request to produce a sample if we are to wait for the texture.
+		// Not all players (Protron, WMF) will produce a texture, unless there is a seek request.
+		if (bWaitForTexture && !bSeekForTextureRequested)
+		{
+			bSeekForTextureRequested = true;
+			InMediaPlayer->Seek(SeekTime);
+		}
+		
+		if (!SeekTime.IsZero() || bSeekForTextureRequested)
+		{
+			if (bSawSeekCompleted)
+			{
+				return UpdateOperationConditionalWaitForTexture(InResponse);
+			}
+			
+			if (SeekTime < FTimespan::FromSeconds(0) || SeekTime > InMediaPlayer->GetDuration())
+			{
+				UE_LOG(LogMediaPlate, Warning, TEXT("[%d] Media Plate Open Latent: Media player seeking to time out of bounds. Seek: %s, Duration: %s, URL: %s"),
+					OutputLink, *SeekTime.ToString(), *InMediaPlayer->GetDuration().ToString(), *URL);
+				FailedOperation(InResponse);
+				return OperationUpdateResult::Completed;
+			}
+			
+			UE_LOG(LogMediaPlate, Verbose, TEXT("[%d] Media Plate Open Latent: Waiting for seek completed event ... (Time out in %f seconds)"), OutputLink, TimeRemaining);
+		}
+		else
+		{
+			return UpdateOperationConditionalWaitForTexture(InResponse);
+		}
+	}
+	else
+	{
+		UE_LOG(LogMediaPlate, Verbose, TEXT("[%d] Media Plate Open Latent: Waiting for opened event ... (Time out in %f seconds)"), OutputLink, TimeRemaining);
+	}
+
+	return OperationUpdateResult::Continue;
+}
+
+FMediaPlateOpenLatentAction::OperationUpdateResult FMediaPlateOpenLatentAction::UpdateOperationConditionalWaitForTexture(FLatentResponse& InResponse) const
+{
+	if (bWaitForTexture)
+	{
+		// Is the texture ready?
+		const UMediaTexture* Texture = MediaPlateComponentWeak->GetMediaTexture();
+		// The aspect ratio will change when we have something.
+		if (Texture && Texture->CurrentAspectRatio != 0.0f)
+		{
+			UE_LOG(LogMediaPlate, Verbose, TEXT("[%d] Media Plate Open Latent: Triggering output pin after media texture ready. Success: %d, %s"), OutputLink, OutSuccess, *URL);
+			CompleteOperation(InResponse);
+			return OperationUpdateResult::Completed;
+		}
+
+		UE_LOG(LogMediaPlate, Verbose, TEXT("[%d] Media Plate Open Latent: Waiting for texture ... (Time out in %f seconds)"), OutputLink, TimeRemaining);	
+		return OperationUpdateResult::Continue;
+	}
+	
+	UE_LOG(LogMediaPlate, Verbose, TEXT("[%d] Media Plate Open Latent: Triggering output pin after seek completed. Success: %d, %s"), OutputLink, OutSuccess, *URL);
+	CompleteOperation(InResponse);
+	return OperationUpdateResult::Completed;
+}
+
+void FMediaPlateOpenLatentAction::OnMediaEvent(EMediaEvent InEvent)
+{
+	UE_LOG(LogMediaPlate, Verbose, TEXT("[%d] Media Plate Open Latent: Saw event: %s"), OutputLink, *MediaUtils::EventToString(InEvent));
+
+	switch (InEvent)
+	{
+	case EMediaEvent::MediaOpened:
+		bSawMediaOpened = true;
+		break;
+	case EMediaEvent::MediaOpenFailed:
+		bSawMediaOpenFailed = true;
+		break;
+	case EMediaEvent::MediaClosed:
+		bSawMediaClosed = true;
+		break;
+	case EMediaEvent::SeekCompleted:
+		bSawSeekCompleted = true;
+		break;
+
+	default:
+		break;
+	}
+}
+
+void FMediaPlateOpenLatentAction::FailedOperation(FLatentResponse& InResponse) const
+{
+	OutSuccess = false;
+	InResponse.FinishAndTriggerIf(true, ExecutionFunction, OutputLink, CallbackTarget);
+}
+
+void FMediaPlateOpenLatentAction::CompleteOperation(FLatentResponse& InResponse) const
+{
+	OutSuccess = true;
+	InResponse.FinishAndTriggerIf(true, ExecutionFunction, OutputLink, CallbackTarget);
+}
\ No newline at end of file
diff --git a/Engine/Plugins/Media/MediaPlate/Source/MediaPlate/Private/MediaPlateOpenLatentAction.h b/Engine/Plugins/Media/MediaPlate/Source/MediaPlate/Private/MediaPlateOpenLatentAction.h
new file mode 100644
--- /dev/null
+++ b/Engine/Plugins/Media/MediaPlate/Source/MediaPlate/Private/MediaPlateOpenLatentAction.h
@@ -0,0 +1,106 @@
+// Copyright Epic Games, Inc. All Rights Reserved.
+
+#pragma once
+
+#include "Engine/LatentActionManager.h"
+#include "IMediaEventSink.h"
+#include "LatentActions.h"
+
+class UMediaPlateComponent;
+class UMediaPlayer;
+
+/**
+ * Pending Latent action to open a media plate component.
+ * The difference with UMediaPlayer::OpenSourceLatent is that is also supports waiting
+ * on the media texture to render a sample and therefore be truly ready to be displayed.
+ * This can be used to hook in other animation systems to manage dynamic loading of media plates.
+ */
+class FMediaPlateOpenLatentAction : public FPendingLatentAction
+{
+public:
+	FMediaPlateOpenLatentAction(const FLatentActionInfo& InLatentInfo,
+		UMediaPlateComponent* InMediaPlateComponent, float InWaitTimeout, bool bInWaitForTexture, bool& InOutSuccess);
+	
+	virtual ~FMediaPlateOpenLatentAction() override;
+	
+	virtual void UpdateOperation(FLatentResponse& InResponse) override;
+
+#if WITH_EDITOR
+	/** Returns a human readable description of the latent operation's current state. */
+	virtual FString GetDescription() const override;
+#endif
+
+private:
+	/** Returns a human readable description of the latent operation's current state. */
+	FString GetStatusString() const;
+
+	/** Indicates if the operation is completed or needs to continue. */
+	enum class OperationUpdateResult
+	{
+		Completed,
+		Continue
+	};
+
+	/** Update Operation when the player is ready. */
+	OperationUpdateResult UpdateOperationPlayerReady(UMediaPlayer* InMediaPlayer, FLatentResponse& InResponse);
+
+	/** Update Operation when the player is ready and done seeking (if waiting for texture). */
+	OperationUpdateResult UpdateOperationConditionalWaitForTexture(FLatentResponse& InResponse) const;
+
+	/** Handler for UMediaPlayer::OnMediaEvent. */
+	void OnMediaEvent(EMediaEvent InEvent);
+
+	/** Ends the operation with failure state. */
+	void FailedOperation(FLatentResponse& InResponse) const;
+
+	/** Completes the operation with success. */
+	void CompleteOperation(FLatentResponse& InResponse) const;
+
+	/** Latent Action Info - The function to execute. */
+	FName ExecutionFunction;
+	
+	/** Latent Action Info - The resume point within the function to execute. */
+	int32 OutputLink;
+
+	/** Latent Action Info - Object to execute the function on. */
+	FWeakObjectPtr CallbackTarget;
+
+	/** Media Plate Component the action is done on. */
+	TWeakObjectPtr<UMediaPlateComponent> MediaPlateComponentWeak;
+
+	/** Media Player the action is done on. */
+	TWeakObjectPtr<UMediaPlayer> MediaPlayerWeak;
+
+	/** Input parameter - If true the action will not be completed until the media texture has rendered a sample. */
+	bool bWaitForTexture = false;
+
+	/** Output parameter - Indicate if the operation completed successfully. */
+	bool& OutSuccess;
+
+	/** Keeps track of the remaining time, in seconds, before the operation times out. */
+	float TimeRemaining = 10.0f;
+
+	/** This is set to true if one of the callbacks indicate an error. */
+	bool bSawError = false;
+
+	/** Set to true if EMediaEvent::MediaOpened is received from the media player. */
+	bool bSawMediaOpened = false;
+
+	/** Set to true if EMediaEvent::MediaClosed is received from the media player. */
+	bool bSawMediaClosed = false;
+
+	/** Set to true if EMediaEvent::MediaOpenFailed is received from the media player. */
+	bool bSawMediaOpenFailed = false;
+
+	/** Set to true if EMediaEvent::SeekCompleted is received from the media player. */
+	bool bSawSeekCompleted = false;
+
+	/** Set to true if the media player has become ready. This is used to detect the state transition. */
+	bool bSawIsReady = false;
+
+	/** Set to true once the seek request has been issued. */
+	bool bSeekForTextureRequested = false;
+
+	/** Cache the media source's URL for display in log messages. */
+	FString URL;
+};
diff --git a/Engine/Plugins/Media/MediaPlate/Source/MediaPlate/Public/MediaPlateComponent.h b/Engine/Plugins/Media/MediaPlate/Source/MediaPlate/Public/MediaPlateComponent.h
--- a/Engine/Plugins/Media/MediaPlate/Source/MediaPlate/Public/MediaPlateComponent.h
+++ b/Engine/Plugins/Media/MediaPlate/Source/MediaPlate/Public/MediaPlateComponent.h
@@ -84,6 +84,7 @@ public:
 #endif // WITH_EDITOR
 	virtual void OnRegister() override;
 	virtual void BeginPlay() override;
+	virtual void EndPlay(const EEndPlayReason::Type InEndPlayReason) override;
 	virtual void BeginDestroy() override;
 	virtual void TickComponent(float DeltaTime, enum ELevelTick TickType, FActorComponentTickFunction* ThisTickFunction) override;
 	virtual void OnUnregister() override;
@@ -120,6 +121,16 @@ public:
 	UFUNCTION(BlueprintCallable, Category = "Media|MediaPlateComponent")
 	void Open();
 
+	/**
+	 * Open the media using a latent action.
+	 *
+	 * @param InTimeout Wait timeout in seconds
+	 * @param bInWaitForTexture Wait for the media texture to have rendered a sample.
+	 * @param bOutSuccess The media was opened successfully.
+	 */
+	UFUNCTION(BlueprintCallable, Category = "Media|MediaPlateComponent", meta = (Latent, LatentInfo="InLatentInfo", WorldContext="InWorldContextObject", InTimeout="10.0f", bInWaitForTexture="true"))	
+	void OpenLatent(const UObject* InWorldContextObject, struct FLatentActionInfo InLatentInfo, float InTimeout, bool bInWaitForTexture, bool& bOutSuccess);
+	
 	/**
 	 * Play the next item in the playlist.
 	 *
