[MediaPlateEditor] Adding an extension point for other plugins to add asset paths to get media plate material templates.

diff --git a/Engine/Plugins/Media/MediaPlate/Source/MediaPlateEditor/Private/MediaPlateEditorModule.cpp b/Engine/Plugins/Media/MediaPlate/Source/MediaPlateEditor/Private/MediaPlateEditorModule.cpp
--- a/Engine/Plugins/Media/MediaPlate/Source/MediaPlateEditor/Private/MediaPlateEditorModule.cpp
+++ b/Engine/Plugins/Media/MediaPlate/Source/MediaPlateEditor/Private/MediaPlateEditorModule.cpp
@@ -155,6 +155,11 @@ bool FMediaPlateEditorModule::RemoveMediaSourceFromDragDropCache(UMediaSource* M
 	return bIsInCache;
 }
 
+FMediaPlateEditorModule::FOnGetMediaPlateMaterialAssetPaths& FMediaPlateEditorModule::OnGetMediaPlateMaterialAssetPaths()
+{
+	return OnGetMediaPlateMaterialAssetPathsDelegate;
+}
+
 const FPlacementCategoryInfo* FMediaPlateEditorModule::GetMediaCategoryRegisteredInfo() const
 {
 	const IPlacementModeModule& PlacementModeModule = IPlacementModeModule::Get();
diff --git a/Engine/Plugins/Media/MediaPlate/Source/MediaPlateEditor/Private/Widgets/SMediaPlateEditorMaterial.cpp b/Engine/Plugins/Media/MediaPlate/Source/MediaPlateEditor/Private/Widgets/SMediaPlateEditorMaterial.cpp
--- a/Engine/Plugins/Media/MediaPlate/Source/MediaPlateEditor/Private/Widgets/SMediaPlateEditorMaterial.cpp
+++ b/Engine/Plugins/Media/MediaPlate/Source/MediaPlateEditor/Private/Widgets/SMediaPlateEditorMaterial.cpp
@@ -10,6 +10,7 @@
 #include "Materials/Material.h"
 #include "Materials/MaterialInstance.h"
 #include "MediaPlate.h"
+#include "MediaPlateEditorModule.h"
 #include "ScopedTransaction.h"
 #include "Widgets/Layout/SBox.h"
 #include "Widgets/Input/SButton.h"
@@ -94,9 +95,24 @@ TSharedRef<SWidget> SMediaPlateEditorMaterial::OnGetMaterials()
 	FMenuBuilder MenuBuilder(true, NULL);
 
 	// Get media plate assets.
-	IAssetRegistry& AssetRegistry = FModuleManager::GetModuleChecked<FAssetRegistryModule>("AssetRegistry").Get();
+	const IAssetRegistry& AssetRegistry = FModuleManager::GetModuleChecked<FAssetRegistryModule>("AssetRegistry").Get();
+
+	TArray<FName> AssetPaths;
+
+	if (FMediaPlateEditorModule* EditorModule = FModuleManager::LoadModulePtr<FMediaPlateEditorModule>("MediaPlateEditor"))
+	{
+		EditorModule->OnGetMediaPlateMaterialAssetPaths().Broadcast(AssetPaths);
+	}
+
+	// Add MediaPlate plugin content path.
+	AssetPaths.AddUnique(FName("/MediaPlate"));
+
 	TArray<FAssetData> Assets;
-	AssetRegistry.GetAssetsByPath(FName("/MediaPlate"), Assets, true);
+	AssetRegistry.GetAssetsByPaths(AssetPaths, Assets, /*bRecursive*/ true);
+
+	// Sort by name to have a consistent list.
+	Algo::Sort(Assets, [](const FAssetData& InAssetA, const FAssetData& InAssetB)
+		{ return InAssetA.AssetName.LexicalLess(InAssetB.AssetName); });
 
 	for (const FAssetData& AssetData : Assets)
 	{
diff --git a/Engine/Plugins/Media/MediaPlate/Source/MediaPlateEditor/Public/MediaPlateEditorModule.h b/Engine/Plugins/Media/MediaPlate/Source/MediaPlateEditor/Public/MediaPlateEditorModule.h
--- a/Engine/Plugins/Media/MediaPlate/Source/MediaPlateEditor/Public/MediaPlateEditorModule.h
+++ b/Engine/Plugins/Media/MediaPlate/Source/MediaPlateEditor/Public/MediaPlateEditorModule.h
@@ -47,6 +47,13 @@ public:
 	 * @return				True if it is a drag and drop media source.
 	 */
 	bool RemoveMediaSourceFromDragDropCache(UMediaSource* MediaSource);
+
+	/**
+	 * Delegate called when getting the media plate materials asset paths.
+	 * This allows for additional plugins to provide new media plate material templates. 
+	 */
+	DECLARE_MULTICAST_DELEGATE_OneParam(FOnGetMediaPlateMaterialAssetPaths, TArray<FName>& /*OutAssetPaths*/);
+	MEDIAPLATEEDITOR_API FOnGetMediaPlateMaterialAssetPaths& OnGetMediaPlateMaterialAssetPaths();
 	
 private:
 
@@ -57,6 +64,9 @@ private:
 	/** Holds all the media plates that are playing. */
 	TArray<TWeakObjectPtr<UMediaPlateComponent>> ActiveMediaPlates;
 
+	/** Event Delegate for OnGetMediaPlateMaterialAssetPaths. */
+	FOnGetMediaPlateMaterialAssetPaths OnGetMediaPlateMaterialAssetPathsDelegate;
+
 	/** Handle for our track editor. */
 	FDelegateHandle TrackEditorBindingHandle;
 
