[MediaPlate] Fix the proxied media texture tracking.

When scrubbing in reverse, it is possible for the "release texture" call to be made after the "get texture" call. The way proxied textures are tracked needs to be order independent. Thus, we rely on "ref counting" instead. This is what the "TextureLayers" are already doing, i.e. each time a TextureIndex is registered it is added in the "Textures" array. So this is equivalent to ref counting, instead of increasing a counter, a new entry with that index is added in the array.

diff --git a/Engine/Plugins/Media/MediaPlate/Source/MediaPlate/Private/MediaPlateComponent.cpp b/Engine/Plugins/Media/MediaPlate/Source/MediaPlate/Private/MediaPlateComponent.cpp
--- a/Engine/Plugins/Media/MediaPlate/Source/MediaPlate/Private/MediaPlateComponent.cpp
+++ b/Engine/Plugins/Media/MediaPlate/Source/MediaPlate/Private/MediaPlateComponent.cpp
@@ -823,13 +823,26 @@ UMediaSource* UMediaPlateComponent::ProxyGetMediaSourceFromIndex(int32 Index) co
 	return UE::MediaPlateComponent::GetMediaSource(GetMediaPlaylist(), Index);
 }
 
+bool UMediaPlateComponent::IsMediaTextureProxied(int32 InTextureIndex) const
+{
+	for (const Layer& TextureLayer : TextureLayers)
+	{
+		for (const int32 TextureIndex : TextureLayer.Textures)
+		{
+			if (TextureIndex == InTextureIndex)
+			{
+				return true;
+			}
+		}
+	}
+	return false;
+}
+
 UMediaTexture* UMediaPlateComponent::ProxyGetMediaTexture(int32 LayerIndex, int32 TextureIndex)
 {
 	UMediaTexture* MediaTexture = GetMediaTexture(TextureIndex);
 	if (MediaTexture != nullptr)
 	{
-		ProxiedMediaTextureIndices.Add(TextureIndex);
-
 		// Ensure internal media player is closed to not interfere with external one (audio in particular).
 		// todo: revisit this when audio proxying is supported. It might be possible to be more permissive.
 		if (IsMediaPlatePlaying())
@@ -867,7 +880,6 @@ UMediaTexture* UMediaPlateComponent::ProxyGetMediaTexture(int32 LayerIndex, int3
 
 void UMediaPlateComponent::ProxyReleaseMediaTexture(int32 LayerIndex, int32 TextureIndex)
 {
-	ProxiedMediaTextureIndices.Remove(TextureIndex);
 	ProxySetTextureBlend(LayerIndex, TextureIndex, 0.0f);
 
 	if (LayerIndex < TextureLayers.Num())
@@ -1226,7 +1238,7 @@ void UMediaPlateComponent::SetUpTextures(UE::MediaPlateComponent::ESetUpTextures
 			// Even if setting the player is allowed, we don't set the player if the texture is currently being proxied.
 			if (EnumHasAllFlags(Flags, UE::MediaPlateComponent::ESetUpTexturesFlags::AllowSetPlayer)
 				&& MediaTexture->GetMediaPlayer() != MediaPlayer.Get()
-				&& !ProxiedMediaTextureIndices.Contains(MediaTextureIndex))
+				&& !IsMediaTextureProxied(MediaTextureIndex))
 			{
 				MediaTexture->SetMediaPlayer(MediaPlayer);
 				bApplyTextureUpdate = true;
@@ -1297,7 +1309,17 @@ bool UMediaPlateComponent::IsNormalMode() const
 
 bool UMediaPlateComponent::IsAnyMediaTextureProxied() const
 {
-	return !ProxiedMediaTextureIndices.IsEmpty();
+	for (const Layer& TextureLayer : TextureLayers)
+	{
+		for (const int32 TextureIndex : TextureLayer.Textures)
+		{
+			if (TextureIndex >= 0)
+			{
+				return true;
+			}
+		}
+	}
+	return false;
 }
 
 void UMediaPlateComponent::UpdateTextureLayers()
diff --git a/Engine/Plugins/Media/MediaPlate/Source/MediaPlate/Public/MediaPlateComponent.h b/Engine/Plugins/Media/MediaPlate/Source/MediaPlate/Public/MediaPlateComponent.h
--- a/Engine/Plugins/Media/MediaPlate/Source/MediaPlate/Public/MediaPlateComponent.h
+++ b/Engine/Plugins/Media/MediaPlate/Source/MediaPlate/Public/MediaPlateComponent.h
@@ -500,11 +500,8 @@ private:
 	};
 	TArray<Layer> TextureLayers;
 
-	/**
-	 * This accurately reflects the set of textures that are actively being used externally
-	 * through the IMediaPlayerProxyInterface's Get/ReleaseMediaTexture functions.
-	 */
-	TSet<int32> ProxiedMediaTextureIndices;
+	/** Determines if the given texture index is being proxied. */
+	bool IsMediaTextureProxied(int32 InTextureIndex) const;
 
 	/**
 	 * Plays a media source.
