[MediaPlate] Fix texture proxy detection.

Regression introduced in CL 41070444.

TextureLayers are used even without proxying, so it can't be used to track proxying. We need a separate data structure to ensure only proxying is being tracked.

diff --git a/Engine/Plugins/Media/MediaPlate/Source/MediaPlate/Private/MediaPlateComponent.cpp b/Engine/Plugins/Media/MediaPlate/Source/MediaPlate/Private/MediaPlateComponent.cpp
--- a/Engine/Plugins/Media/MediaPlate/Source/MediaPlate/Private/MediaPlateComponent.cpp
+++ b/Engine/Plugins/Media/MediaPlate/Source/MediaPlate/Private/MediaPlateComponent.cpp
@@ -823,18 +823,31 @@ UMediaSource* UMediaPlateComponent::ProxyGetMediaSourceFromIndex(int32 Index) co
 	return UE::MediaPlateComponent::GetMediaSource(GetMediaPlaylist(), Index);
 }
 
-bool UMediaPlateComponent::IsMediaTextureProxied(int32 InTextureIndex) const
+void UMediaPlateComponent::IncreaseMediaTextureProxyCount(int32 InTextureIndex)
 {
-	for (const Layer& TextureLayer : TextureLayers)
+	int32& ProxyCount = MediaTextureProxyCount.FindOrAdd(InTextureIndex, 0);
+	++ProxyCount;
+}
+
+void UMediaPlateComponent::DecreaseMediaTextureProxyCount(int32 InTextureIndex)
+{
+	if (int32* ProxyCount = MediaTextureProxyCount.Find(InTextureIndex))
 	{
-		for (const int32 TextureIndex : TextureLayer.Textures)
+		--(*ProxyCount);
+
+		if (*ProxyCount <= 0)
 		{
-			if (TextureIndex == InTextureIndex)
-			{
-				return true;
-			}
+			MediaTextureProxyCount.Remove(InTextureIndex);
 		}
 	}
+}
+
+bool UMediaPlateComponent::IsMediaTextureProxied(int32 InTextureIndex) const
+{
+	if (const int32* RefCount = MediaTextureProxyCount.Find(InTextureIndex))
+	{
+		return *RefCount > 0;
+	}
 	return false;
 }
 
@@ -843,6 +856,8 @@ UMediaTexture* UMediaPlateComponent::ProxyGetMediaTexture(int32 LayerIndex, int3
 	UMediaTexture* MediaTexture = GetMediaTexture(TextureIndex);
 	if (MediaTexture != nullptr)
 	{
+		IncreaseMediaTextureProxyCount(TextureIndex);
+		
 		// Ensure internal media player is closed to not interfere with external one (audio in particular).
 		// todo: revisit this when audio proxying is supported. It might be possible to be more permissive.
 		if (IsMediaPlatePlaying())
@@ -880,6 +895,8 @@ UMediaTexture* UMediaPlateComponent::ProxyGetMediaTexture(int32 LayerIndex, int3
 
 void UMediaPlateComponent::ProxyReleaseMediaTexture(int32 LayerIndex, int32 TextureIndex)
 {
+	DecreaseMediaTextureProxyCount(TextureIndex);
+	
 	ProxySetTextureBlend(LayerIndex, TextureIndex, 0.0f);
 
 	if (LayerIndex < TextureLayers.Num())
@@ -1309,14 +1326,11 @@ bool UMediaPlateComponent::IsNormalMode() const
 
 bool UMediaPlateComponent::IsAnyMediaTextureProxied() const
 {
-	for (const Layer& TextureLayer : TextureLayers)
+	for (const TPair<int32, int32>& TextureProxyCount : MediaTextureProxyCount)
 	{
-		for (const int32 TextureIndex : TextureLayer.Textures)
+		if (TextureProxyCount.Value > 0)
 		{
-			if (TextureIndex >= 0)
-			{
-				return true;
-			}
+			return true;
 		}
 	}
 	return false;
diff --git a/Engine/Plugins/Media/MediaPlate/Source/MediaPlate/Public/MediaPlateComponent.h b/Engine/Plugins/Media/MediaPlate/Source/MediaPlate/Public/MediaPlateComponent.h
--- a/Engine/Plugins/Media/MediaPlate/Source/MediaPlate/Public/MediaPlateComponent.h
+++ b/Engine/Plugins/Media/MediaPlate/Source/MediaPlate/Public/MediaPlateComponent.h
@@ -500,6 +500,23 @@ private:
 	};
 	TArray<Layer> TextureLayers;
 
+	/**
+	 * Map of TextureIndex to proxy count.
+	 * 
+	 * This accurately reflects the set of textures that are actively being used externally
+	 * through the IMediaPlayerProxyInterface's Get/ReleaseMediaTexture functions.
+	 * Given that Get/Release calls can be out of order depending on the playback or scrubbing direction,
+	 * we need to keep track of this state using a ref counting method. Some texture can temporarily be
+	 * referenced by 2 (or more) sections (even if not actually used internally).
+	 */
+	TMap<int32, int32> MediaTextureProxyCount;
+
+	/** Called when the given media texture starts being proxied by a media section to increase the proxy count. */ 
+	void IncreaseMediaTextureProxyCount(int32 InTextureIndex);
+
+	/** Called when the given media texture stops being proxied by a media section to decrease the proxy count. */ 
+	void DecreaseMediaTextureProxyCount(int32 InTextureIndex);
+
 	/** Determines if the given texture index is being proxied. */
 	bool IsMediaTextureProxied(int32 InTextureIndex) const;
 
