[MediaPlateEditor] Fixing the pause button being shown when the player is closed.

This causes the buttons to move around because both pause and play are visible when the player is closed. Then when opening the player (with auto play), the play button disapears and cause the other buttons to move left, which is not ideal.
The visible/collapse logic of play button must be exact opposite of pause button.

diff --git a/Engine/Plugins/Media/MediaPlate/Source/MediaPlateEditor/Private/MediaPlateCustomization.cpp b/Engine/Plugins/Media/MediaPlate/Source/MediaPlateEditor/Private/MediaPlateCustomization.cpp
--- a/Engine/Plugins/Media/MediaPlate/Source/MediaPlateEditor/Private/MediaPlateCustomization.cpp
+++ b/Engine/Plugins/Media/MediaPlate/Source/MediaPlateEditor/Private/MediaPlateCustomization.cpp
@@ -239,9 +239,11 @@ void FMediaPlateCustomization::CustomizeDetails(IDetailLayoutBuilder& DetailBuil
 						{
 							if (const TSharedPtr<FMediaPlateCustomization> Self = WeakSelf.Pin())
 							{
+								// We want this logic to be as mutually exclusive with the play button visibility as possible so
+								// they don't show at the same time and cause other buttons to move around.
 								const bool bAllPaused = Self->IsTrueForAllPlayers([](const UMediaPlayer* InMediaPlayer)
 								{
-									return InMediaPlayer->IsPaused();
+									return !InMediaPlayer->IsPlaying(); // Not using IsPaused() as it is not the logical inverse of IsPlaying.
 								});
 								return bAllPaused ? EVisibility::Collapsed : EVisibility::Visible;
 							}
