[MediaPlate] Fix MediaPlateResource to be able to include it in a blueprint structure.

- Need to move the instanced properties out of the struct and to the MediaPlateComponent instead.
- Change the customization to support editing a stand alone MediaPlateResource (with no outer).

Bonus:
- Added SelectMediaSourceAsset, SelectExternalMedia and SelectMediaPlaylist functions to make it easier to set the asset in BP.

diff --git a/Engine/Plugins/Media/MediaPlate/Source/MediaPlate/Private/MediaPlateComponent.cpp b/Engine/Plugins/Media/MediaPlate/Source/MediaPlate/Private/MediaPlateComponent.cpp
--- a/Engine/Plugins/Media/MediaPlate/Source/MediaPlate/Private/MediaPlateComponent.cpp
+++ b/Engine/Plugins/Media/MediaPlate/Source/MediaPlate/Private/MediaPlateComponent.cpp
@@ -155,6 +155,18 @@ PRAGMA_DISABLE_DEPRECATION_WARNINGS
 		InitializeMediaPlateResource();
 	}
 PRAGMA_ENABLE_DEPRECATION_WARNINGS
+	
+	// Ensure the external media source is loaded.
+	if (MediaPlateResource.Type == EMediaPlateResourceType::External && !ExternalMediaSource)
+	{
+		RefreshExternalMediaSource();
+	}
+
+	// Ensure active playlist is refreshed.
+	if (!ActivePlaylist)
+	{
+		RefreshActivePlaylist();
+	}
 }
 #endif // WITH_EDITOR
 
@@ -509,7 +521,25 @@ void UMediaPlateComponent::SetLoop(bool bInLoop)
 
 UMediaPlaylist* UMediaPlateComponent::GetMediaPlaylist() const
 {
-	return MediaPlateResource.GetActivePlaylist();
+	return ActivePlaylist;
+}
+
+void UMediaPlateComponent::SelectExternalMedia(const FString& InFilePath)
+{
+	MediaPlateResource.LoadExternalMedia(InFilePath);
+	RefreshMediaPlateResource();
+}
+
+void UMediaPlateComponent::SelectMediaSourceAsset(const UMediaSource* InMediaSource)
+{
+	MediaPlateResource.SelectAsset(InMediaSource);
+	RefreshMediaPlateResource();
+}
+
+void UMediaPlateComponent::SelectMediaPlaylistAsset(const UMediaPlaylist* InMediaPlaylist)
+{
+	MediaPlateResource.SelectPlaylist(InMediaPlaylist);
+	RefreshMediaPlateResource();
 }
 
 void UMediaPlateComponent::SetMediaPlateResource(const FMediaPlateResource& InMediaPlayerResource)
@@ -517,16 +547,13 @@ void UMediaPlateComponent::SetMediaPlateResource(const FMediaPlateResource& InMe
 #if WITH_EDITOR
 	const FScopedTransaction Transaction(LOCTEXT("SetMediaPlateResource", "Set media plate resource"));
 	this->Modify();
-	MediaPlateResource.Modify();
-#endif
-
-	MediaPlateResource.Init(InMediaPlayerResource);
-
-	if (MediaPlateResource.GetResourceType() == EMediaPlateResourceType::External)
+	if (GetMediaPlaylist())
 	{
-		MediaPlateResource.LoadExternalMedia(FString(InMediaPlayerResource.GetExternalMediaPath()), this);
+		GetMediaPlaylist()->Modify();
 	}
+#endif
 
+	MediaPlateResource.Init(InMediaPlayerResource);
 	RefreshMediaPlateResource();
 }
 
@@ -1425,11 +1452,11 @@ void UMediaPlateComponent::InitializeMediaPlateResource()
 	if (PlaylistElementsNum > 1)
 	{
 		// Playlist has more than one element, source type will be set to playlist
-		MediaPlateResource.SelectPlaylist(Playlist);
+		SelectMediaPlaylistAsset(Playlist);
 	}
 	else
 	{
-		if (UMediaSource* MediaSource = Playlist->Get(0))
+		if (const UMediaSource* MediaSource = Playlist->Get(0))
 		{
 			// If outer is this, source has been created from external file
 			if (MediaSource->GetOuter() == this)
@@ -1447,14 +1474,11 @@ void UMediaPlateComponent::InitializeMediaPlateResource()
 					MediaPath = MediaPath.RightChop(ImgPrefix.Len());
 				}
 
-				MediaPlateResource.ExternalMediaPath = MediaPath;
-				MediaPlateResource.ExternalMedia = MediaSource;
-				MediaPlateResource.SetResourceType(EMediaPlateResourceType::External);
-				MediaPlateResource.RefreshActivePlaylist(this);
+				SelectExternalMedia(MediaPath);
 			}
 			else
 			{
-				MediaPlateResource.SelectAsset(MediaSource, this);
+				SelectMediaSourceAsset(MediaSource);
 			}
 		}
 	}
@@ -1466,7 +1490,8 @@ void UMediaPlateComponent::InitializeMediaPlateResource()
 
 void UMediaPlateComponent::RefreshMediaPlateResource()
 {
-	MediaPlateResource.RefreshActivePlaylist(this);
+	RefreshExternalMediaSource();
+	RefreshActivePlaylist();
 
 	// Remote Control may call PostEditChangeProperty more than once from various paths
 	// with some delays causing Media Plate to close when it shouldn't. We mitigate that by
@@ -1495,6 +1520,47 @@ void UMediaPlateComponent::RefreshMediaPlateResource()
 	}
 }
 
+void UMediaPlateComponent::RefreshExternalMediaSource()
+{
+	if (MediaPlateResource.GetResourceType() == EMediaPlateResourceType::External)
+	{
+		ExternalMediaSource = UMediaSource::SpawnMediaSourceForString(MediaPlateResource.ExternalMediaPath, this);
+	}
+}
+
+void UMediaPlateComponent::RefreshActivePlaylist()
+{
+	UMediaSource* SelectedMediaSource = nullptr;
+	
+	switch (MediaPlateResource.GetResourceType())
+	{
+	case EMediaPlateResourceType::Playlist:
+		ActivePlaylist = MediaPlateResource.GetSourcePlaylist();
+		return;
+
+	case EMediaPlateResourceType::External:
+		if (!ExternalMediaSource)
+		{
+			RefreshExternalMediaSource();
+		}
+		SelectedMediaSource = ExternalMediaSource;
+		break;
+
+	case EMediaPlateResourceType::Asset:
+		SelectedMediaSource = MediaPlateResource.GetMediaAsset();
+		break;
+
+	default:
+		break;
+	}
+
+	if (SelectedMediaSource)
+	{
+		ActivePlaylist = NewObject<UMediaPlaylist>(this, NAME_None, RF_Transactional);
+		ActivePlaylist->Add(SelectedMediaSource);
+	}
+}
+
 void UMediaPlateComponent::RefreshMediaSoundComponent()
 {
 	// Are we turning on audio?
@@ -1629,12 +1695,34 @@ void UMediaPlateComponent::PostEditChangeProperty(FPropertyChangedEvent& Propert
 	}
 }
 
+void UMediaPlateComponent::PostEditChangeChainProperty(FPropertyChangedChainEvent& InPropertyChangedEvent)
+{
+	Super::PostEditChangeChainProperty(InPropertyChangedEvent);
+
+	static const FName TypePropertyName = GET_MEMBER_NAME_CHECKED(FMediaPlateResource, Type);
+	static const FName MediaAssetPropertyName = GET_MEMBER_NAME_CHECKED(FMediaPlateResource, MediaAsset);
+	static const FName ExternalMediaPathPropertyName = GET_MEMBER_NAME_CHECKED(FMediaPlateResource, ExternalMediaPath);
+	static const FName SourcePlaylistPropertyName = GET_MEMBER_NAME_CHECKED(FMediaPlateResource, SourcePlaylist);
+
+	const FName MemberName = InPropertyChangedEvent.GetMemberPropertyName();
+	
+	if (MemberName == TypePropertyName
+		|| MemberName == MediaAssetPropertyName
+		|| MemberName == ExternalMediaPathPropertyName
+		|| MemberName == SourcePlaylistPropertyName)
+	{
+		RefreshMediaPlateResource();
+	}
+}
+
 void UMediaPlateComponent::PostEditUndo()
 {
 	Super::PostEditUndo();
 	
 	// Try to salvage the restored component. 
 	RefreshMediaSoundComponent();
+
+	RefreshMediaPlateResource();
 }
 #endif // WITH_EDITOR
 
diff --git a/Engine/Plugins/Media/MediaPlate/Source/MediaPlate/Private/MediaPlateResource.cpp b/Engine/Plugins/Media/MediaPlate/Source/MediaPlate/Private/MediaPlateResource.cpp
--- a/Engine/Plugins/Media/MediaPlate/Source/MediaPlate/Private/MediaPlateResource.cpp
+++ b/Engine/Plugins/Media/MediaPlate/Source/MediaPlate/Private/MediaPlateResource.cpp
@@ -1,54 +1,36 @@
 // Copyright Epic Games, Inc. All Rights Reserved.
 
 #include "MediaPlateResource.h"
+
 #include "MediaPlaylist.h"
 #include "MediaSource.h"
 
 void FMediaPlateResource::SetResourceType(EMediaPlateResourceType InType)
 {
 	Type = InType;
-
-	if (const UMediaPlaylist* Playlist = GetActivePlaylist())
-	{
-		RefreshActivePlaylist(Playlist->GetOuter());
-	}
 }
 
-void FMediaPlateResource::SelectAsset(const UMediaSource* InMediaSource, UObject* InOuter)
+void FMediaPlateResource::SelectAsset(const UMediaSource* InMediaSource)
 {
 	SetResourceType(EMediaPlateResourceType::Asset);
 	MediaAsset = const_cast<UMediaSource*>(InMediaSource);
-	RefreshActivePlaylist(InOuter);
 }
 
-void FMediaPlateResource::LoadExternalMedia(const FString& InFilePath, UObject* InOuter)
+void FMediaPlateResource::LoadExternalMedia(const FString& InFilePath)
 {
 	SetResourceType(EMediaPlateResourceType::External);
 	ExternalMediaPath = InFilePath;
-	ExternalMedia = UMediaSource::SpawnMediaSourceForString(InFilePath, InOuter);
-	RefreshActivePlaylist(InOuter);
 }
 
 void FMediaPlateResource::SelectPlaylist(const UMediaPlaylist* InPlaylist)
 {
 	SetResourceType(EMediaPlateResourceType::Playlist);
 	SourcePlaylist = const_cast<UMediaPlaylist*>(InPlaylist);
-	ActivePlaylist = GetSourcePlaylist();
-}
-
-#if WITH_EDITOR
-void FMediaPlateResource::Modify() const
-{
-	if (ActivePlaylist)
-	{
-		ActivePlaylist->Modify();
-	}
 }
-#endif
 
 void FMediaPlateResource::Init(const FMediaPlateResource& InOther)
 {
-	if (!ExternalMediaPath.IsEmpty())
+	if (!InOther.GetExternalMediaPath().IsEmpty())
 	{
 		ExternalMediaPath = FString(InOther.GetExternalMediaPath());
 	}
@@ -66,23 +48,6 @@ void FMediaPlateResource::Init(const FMediaPlateResource& InOther)
 	Type = InOther.GetResourceType();
 }
 
-void FMediaPlateResource::RefreshActivePlaylist(UObject* InOuter)
-{
-	if (Type == EMediaPlateResourceType::Playlist)
-	{
-		ActivePlaylist = GetSourcePlaylist();
-	}
-	else if (UMediaSource* MediaSource = GetSelectedMedia())
-	{
-		if (InOuter)
-		{
-			ActivePlaylist = nullptr;
-			ActivePlaylist = NewObject<UMediaPlaylist>(InOuter, NAME_None, RF_Transactional);
-			ActivePlaylist->Add(MediaSource);
-		}
-	}
-}
-
 UMediaSource* FMediaPlateResource::GetMediaAsset() const
 {
 	return MediaAsset.LoadSynchronous();
@@ -91,27 +56,4 @@ UMediaSource* FMediaPlateResource::GetMediaAsset() const
 UMediaPlaylist* FMediaPlateResource::GetSourcePlaylist() const
 {
 	return SourcePlaylist.LoadSynchronous();
-}
-
-UMediaPlaylist* FMediaPlateResource::GetActivePlaylist() const
-{
-	return ActivePlaylist;
-}
-
-UMediaSource* FMediaPlateResource::GetSelectedMedia() const
-{
-	switch (Type)
-	{
-		case EMediaPlateResourceType::Playlist:
-			return nullptr;
-
-		case EMediaPlateResourceType::External:
-			return ExternalMedia;
-
-		case EMediaPlateResourceType::Asset:
-			return MediaAsset.LoadSynchronous();
-
-		default:
-			return nullptr;
-	}
-}
+}
\ No newline at end of file
diff --git a/Engine/Plugins/Media/MediaPlate/Source/MediaPlate/Public/MediaPlateComponent.h b/Engine/Plugins/Media/MediaPlate/Source/MediaPlate/Public/MediaPlateComponent.h
--- a/Engine/Plugins/Media/MediaPlate/Source/MediaPlate/Public/MediaPlateComponent.h
+++ b/Engine/Plugins/Media/MediaPlate/Source/MediaPlate/Public/MediaPlateComponent.h
@@ -91,6 +91,7 @@ public:
 	virtual void OnComponentDestroyed(bool bDestroyingHierarchy) override;
 #if WITH_EDITOR
 	virtual void PostEditChangeProperty(FPropertyChangedEvent& PropertyChangedEvent) override;
+	virtual void PostEditChangeChainProperty(FPropertyChangedChainEvent& InPropertyChangedEvent) override;
 	virtual void PostEditUndo() override;
 #endif // WITH_EDITOR
 
@@ -207,6 +208,27 @@ public:
 	 */
 	UMediaPlaylist* GetMediaPlaylist() const;
 
+	/**
+	 * @brief Select the external media file (non-UE asset) to be opened.
+	 * @param InFilePath non-UFS path (absolute or relative to executable) to the media file.
+	 */
+	UFUNCTION(BlueprintCallable, Category = "Media|MediaPlateComponent")
+	void SelectExternalMedia(const FString& InFilePath);
+	
+	/**
+	 * @brief Select the media source asset to be opened. 
+	 * @param InMediaSource media source asset to open.
+	 */
+	UFUNCTION(BlueprintCallable, Category = "Media|MediaPlateComponent")
+	void SelectMediaSourceAsset(const UMediaSource* InMediaSource);
+
+	/**
+	 * @brief Select the media playlist asset to be opened. 
+	 * @param InMediaPlaylist media playlist asset to open.
+	 */
+	UFUNCTION(BlueprintCallable, Category = "Media|MediaPlateComponent")
+	void SelectMediaPlaylistAsset(const UMediaPlaylist* InMediaPlaylist);
+	
 	/**
 	 * Update Media Player Resource. This will also refresh Playlist accordingly.
 	 */
@@ -467,6 +489,17 @@ private:
 	UPROPERTY(Instanced)
 	TObjectPtr<UMediaPlayer> MediaPlayer;
 
+	/** Embedded Media Source loaded from external path specified in MediaPlateResource. */
+	UPROPERTY(Instanced)
+	TObjectPtr<UMediaSource> ExternalMediaSource;
+
+	/**
+	 * Currently running Playlist.
+	 * It can either be a stand alone or embedded asset depending on the media plate resource type. 
+	 */
+	UPROPERTY(Instanced)
+	TObjectPtr<UMediaPlaylist> ActivePlaylist;
+	
 	/** Info representing this object. */
 	TSharedPtr<FMediaTextureTrackerObject, ESPMode::ThreadSafe> MediaTextureTrackerObject;
 	/** Our media clock sink. */
@@ -646,6 +679,16 @@ private:
 	 */
 	void RefreshMediaPlateResource();
 
+	/**
+	 * If the current resource type is external path, this will load a new media source from that path.
+	 */
+	void RefreshExternalMediaSource();
+	
+	/**
+	 * Refresh the ActivePlaylist from the current configuration of MediaPlateResource.
+	 */
+	void RefreshActivePlaylist();
+	
 	/**
 	 * Refresh the media sound component according to the bEnableAudio variable.
 	 */
diff --git a/Engine/Plugins/Media/MediaPlate/Source/MediaPlate/Public/MediaPlateResource.h b/Engine/Plugins/Media/MediaPlate/Source/MediaPlate/Public/MediaPlateResource.h
--- a/Engine/Plugins/Media/MediaPlate/Source/MediaPlate/Public/MediaPlateResource.h
+++ b/Engine/Plugins/Media/MediaPlate/Source/MediaPlate/Public/MediaPlateResource.h
@@ -32,22 +32,11 @@ struct FMediaPlateResource
 	GENERATED_BODY()
 
 public:
-	/**
-	 * Returns the currently selected Media Source, based on current source Type.
-	 * If current type is Playlist, it will return nullptr
-	 */
-	MEDIAPLATE_API UMediaSource* GetSelectedMedia() const;
-
 	/**
 	 * Returns the current source playlist, if any
 	 */
 	MEDIAPLATE_API UMediaPlaylist* GetSourcePlaylist() const;
 
-	/**
-	 * Returns the active playlist, if any
-	 */
-	MEDIAPLATE_API UMediaPlaylist* GetActivePlaylist() const;
-
 	/**
 	 * Returns the current external media path, if any
 	 */
@@ -71,57 +60,63 @@ public:
 	/**
 	 * Select asset based media source. Will also update source type to Asset
 	 */
-	MEDIAPLATE_API void SelectAsset(const UMediaSource* InMediaSource, UObject* InOuter);
+	MEDIAPLATE_API void SelectAsset(const UMediaSource* InMediaSource);
 
 	/**
 	 * Loads the external source at the specified path, creating a media source with the specified outer.
 	 * Will also update source type to External
 	 */
-	MEDIAPLATE_API void LoadExternalMedia(const FString& InFilePath, UObject* InOuter);
+	MEDIAPLATE_API void LoadExternalMedia(const FString& InFilePath);
 
 	/**
 	 * Select the specified playlist. Will also update source type to Playlist
 	 */
 	MEDIAPLATE_API void SelectPlaylist(const UMediaPlaylist* InPlaylist);
 
+	UE_DEPRECATED(5.6, "Use UMediaPlateComponent::GetSelectedMediaSource() instead.")
+	UMediaSource* GetSelectedMedia() const { return nullptr; }
+
+	UE_DEPRECATED(5.6, "Use UMediaPlateComponent::GetMediaPlaylist() instead.")
+	UMediaPlaylist* GetActivePlaylist() const { return nullptr; }
+
+	UE_DEPRECATED(5.6, "Use UMediaPlateComponent::SelectMediaSourceAsset() instead.")
+	void SelectAsset(const UMediaSource* InMediaSource, UObject* InOuter)
+	{
+		SelectAsset(InMediaSource);
+	}
+
+	UE_DEPRECATED(5.6, "Use UMediaPlateComponent::SelectExternalMedia() instead.")
+	void LoadExternalMedia(const FString& InFilePath, UObject* InOuter)
+	{
+		LoadExternalMedia(InFilePath);
+	}
+
 #if WITH_EDITOR
-	/**
-	 * Note that the Active Playlist will be modified.
-	 * Convenience function to avoid having to check and get the Playlist everytime to call Modify() on it.
-	 */
-	MEDIAPLATE_API void Modify() const;
-#endif
+	UE_DEPRECATED(5.6, "Use UMediaPlateComponent::GetMediaPlaylist()->Modify() instead")
+	void Modify() const {}
+#endif	
 
 private:
 friend class UMediaPlateComponent;
+friend class FMediaPlateResourceCustomization;
 	/**
 	 * Initialize member properties from another MediaPlayerResource. Empty or null property will not be copied over.
 	 */
 	void Init(const FMediaPlateResource& InOther);
 
-	void RefreshActivePlaylist(UObject* InOuter);
-
 	/** Media Source Type */
-	UPROPERTY()
+	UPROPERTY(EditAnywhere, Category = "MediaPlateResource")
 	EMediaPlateResourceType Type = EMediaPlateResourceType::Asset;
 
 	/** A path pointing to an external media resource */
-	UPROPERTY()
+	UPROPERTY(EditAnywhere, Category = "MediaPlateResource")
 	FString ExternalMediaPath;
 
-	/** Media Source loaded from external path */
-	UPROPERTY(Instanced)
-	TObjectPtr<UMediaSource> ExternalMedia;
-
 	/** Media source coming from MediaSource asset*/
-	UPROPERTY()
+	UPROPERTY(EditAnywhere, Category = "MediaPlateResource")
 	TSoftObjectPtr<UMediaSource> MediaAsset;
 
 	/** User facing Playlist asset */
-	UPROPERTY()
+	UPROPERTY(EditAnywhere, Category = "MediaPlateResource")
 	TSoftObjectPtr<UMediaPlaylist> SourcePlaylist;
-
-	/** Currently running Playlist asset */
-	UPROPERTY(Instanced)
-	TObjectPtr<UMediaPlaylist> ActivePlaylist;
-};
+};
\ No newline at end of file
diff --git a/Engine/Plugins/Media/MediaPlate/Source/MediaPlateEditor/Private/ActorFactoryMediaPlate.cpp b/Engine/Plugins/Media/MediaPlate/Source/MediaPlateEditor/Private/ActorFactoryMediaPlate.cpp
--- a/Engine/Plugins/Media/MediaPlate/Source/MediaPlateEditor/Private/ActorFactoryMediaPlate.cpp
+++ b/Engine/Plugins/Media/MediaPlate/Source/MediaPlateEditor/Private/ActorFactoryMediaPlate.cpp
@@ -77,7 +77,7 @@ void UActorFactoryMediaPlate::SetUpActor(UObject* Asset, AActor* Actor)
 								MediaSource->Rename(nullptr, MediaPlateComponent);
 
 								// Let's initialize the plate source with the now non-transient MediaSource
-								MediaPlateComponent->MediaPlateResource.SelectAsset(MediaSource, MediaPlateComponent);
+								MediaPlateComponent->SelectMediaSourceAsset(MediaSource);
 							}
 						}
 					});
@@ -85,7 +85,7 @@ void UActorFactoryMediaPlate::SetUpActor(UObject* Asset, AActor* Actor)
 				else
 				{
 					// MediaSource is non-transient, we can initialize the plate source right away
-					MediaPlate->MediaPlateComponent->MediaPlateResource.SelectAsset(MediaSource, MediaPlate->MediaPlateComponent);
+					MediaPlate->MediaPlateComponent->SelectMediaSourceAsset(MediaSource);
 				}
 			}
 		}
diff --git a/Engine/Plugins/Media/MediaPlate/Source/MediaPlateEditor/Private/MediaPlateResourceCustomization.cpp b/Engine/Plugins/Media/MediaPlate/Source/MediaPlateEditor/Private/MediaPlateResourceCustomization.cpp
--- a/Engine/Plugins/Media/MediaPlate/Source/MediaPlateEditor/Private/MediaPlateResourceCustomization.cpp
+++ b/Engine/Plugins/Media/MediaPlate/Source/MediaPlateEditor/Private/MediaPlateResourceCustomization.cpp
@@ -24,8 +24,13 @@ void FMediaPlateResourceCustomization::CustomizeHeader(TSharedRef<IPropertyHandl
 	const TAttribute<EVisibility> MediaSourceFileVisibility(this, &FMediaPlateResourceCustomization::GetFileSelectorVisibility);
 	const TAttribute<EVisibility> MediaSourceAssetVisibility(this, &FMediaPlateResourceCustomization::GetAssetSelectorVisibility);
 	const TAttribute<EVisibility> MediaSourcePlaylistVisibility(this, &FMediaPlateResourceCustomization::GetPlaylistSelectorVisibility);
+	const TAttribute<EVisibility> MultipleValuesVisibility(this, &FMediaPlateResourceCustomization::GetMultipleValuesVisibility);
 
 	MediaPlateResourcePropertyHandle = InStructPropertyHandle;
+	ResourceTypePropertyHandle = MediaPlateResourcePropertyHandle->GetChildHandle(GET_MEMBER_NAME_CHECKED(FMediaPlateResource, Type));
+	ExternalMediaPathPropertyHandle = MediaPlateResourcePropertyHandle->GetChildHandle(GET_MEMBER_NAME_CHECKED(FMediaPlateResource, ExternalMediaPath));
+	MediaAssetPropertyHandle = MediaPlateResourcePropertyHandle->GetChildHandle(GET_MEMBER_NAME_CHECKED(FMediaPlateResource, MediaAsset));
+	SourcePlaylistPropertyHandle = MediaPlateResourcePropertyHandle->GetChildHandle(GET_MEMBER_NAME_CHECKED(FMediaPlateResource, SourcePlaylist));
 
 	void* StructRawData;
 	// Try accessing the raw value, so we make sure there's only one struct edited. Multiple Access edit is currently not supported
@@ -33,29 +38,31 @@ void FMediaPlateResourceCustomization::CustomizeHeader(TSharedRef<IPropertyHandl
 
 	TSharedPtr<SWidget> ValueWidgetContent;
 
-	if (AccessResult == FPropertyAccess::Success)
+	if (AccessResult == FPropertyAccess::Success || AccessResult == FPropertyAccess::MultipleValues)
 	{
 		ValueWidgetContent = SNew(SBox)
 		[
 			SNew(SVerticalBox)
 			+ SVerticalBox::Slot()
 			.HAlign(HAlign_Left)
+			.MinHeight(22)
+			.MaxHeight(22)
 			[
-				SNew(SSegmentedControl<EMediaPlateResourceType>)
+				SNew(SSegmentedControl<TOptional<EMediaPlateResourceType>>)
 					.Value(this, &FMediaPlateResourceCustomization::GetAssetType)
 					.OnValueChanged(this, &FMediaPlateResourceCustomization::OnAssetTypeChanged)
 
-				+ SSegmentedControl<EMediaPlateResourceType>::Slot(EMediaPlateResourceType::External)
+				+ SSegmentedControl<TOptional<EMediaPlateResourceType>>::Slot(EMediaPlateResourceType::External)
 					.Text(LOCTEXT("File", "File"))
 					.ToolTip(LOCTEXT("File_ToolTip",
 						"Select this if you want to use a file path to a media file on disk."))
 
-				+ SSegmentedControl<EMediaPlateResourceType>::Slot(EMediaPlateResourceType::Asset)
+				+ SSegmentedControl<TOptional<EMediaPlateResourceType>>::Slot(EMediaPlateResourceType::Asset)
 					.Text(LOCTEXT("Asset", "Asset"))
 					.ToolTip(LOCTEXT("Asset_ToolTip",
 						"Select this if you want to use a Media Source asset."))
 
-				+ SSegmentedControl<EMediaPlateResourceType>::Slot(EMediaPlateResourceType::Playlist)
+				+ SSegmentedControl<TOptional<EMediaPlateResourceType>>::Slot(EMediaPlateResourceType::Playlist)
 					.Text(LOCTEXT("Playlist", "Playlist"))
 					.ToolTip(LOCTEXT("Playlist_ToolTip",
 						"Select this if you want to use a Media Playlist asset."))
@@ -68,8 +75,7 @@ void FMediaPlateResourceCustomization::CustomizeHeader(TSharedRef<IPropertyHandl
 				[
 					SNew(SObjectPropertyEntryBox)
 					.AllowedClass(UMediaSource::StaticClass())
-					.ObjectPath(this, &FMediaPlateResourceCustomization::GetMediaAssetPath)
-					.OnObjectChanged(this, &FMediaPlateResourceCustomization::OnMediaAssetChanged)
+					.PropertyHandle(MediaAssetPropertyHandle)
 				]
 			]
 			+ SVerticalBox::Slot()
@@ -96,20 +102,25 @@ void FMediaPlateResourceCustomization::CustomizeHeader(TSharedRef<IPropertyHandl
 				[
 					SNew(SObjectPropertyEntryBox)
 					.AllowedClass(UMediaPlaylist::StaticClass())
-					.ObjectPath(this, &FMediaPlateResourceCustomization::GetPlaylistPath)
-					.OnObjectChanged(this, &FMediaPlateResourceCustomization::OnPlaylistChanged)
+					.PropertyHandle(SourcePlaylistPropertyHandle)
+				]
+			]
+			+ SVerticalBox::Slot()
+			[
+				SNew(SBox)
+				.Visibility(MultipleValuesVisibility)
+				.HAlign(HAlign_Fill)
+				.VAlign(VAlign_Center)
+				.Padding(0, 4)
+				[
+					SNew(STextBlock)
+					.Text(LOCTEXT("MultipleValues", "Multiple Values"))
+					.ToolTipText(LOCTEXT("MultipleValues_ToolTip", "Multiple Values can't be displayed."))
+					.Font(IDetailLayoutBuilder::GetDetailFont())
 				]
 			]
 		];
 	}
-	else if (AccessResult == FPropertyAccess::MultipleValues)
-	{
-		ValueWidgetContent = SNew(STextBlock)
-			.Text(LOCTEXT("MultipleValues", "Multiple Selection"))
-			.ToolTipText(LOCTEXT("MultipleValues_ToolTip",
-			"Multiple Media Player Resource properties selected. Select a single property to edit it."))
-			.Font(IDetailLayoutBuilder::GetDetailFont());
-	}
 	else
 	{
 		ValueWidgetContent = SNew(STextBlock)
@@ -133,186 +144,52 @@ void FMediaPlateResourceCustomization::CustomizeChildren(TSharedRef<IPropertyHan
 {
 }
 
-EMediaPlateResourceType FMediaPlateResourceCustomization::GetAssetType() const
-{
-	constexpr EMediaPlateResourceType DefaultType = EMediaPlateResourceType::Asset;
-
-	if (const FMediaPlateResource* MediaPlateResource = GetMediaPlateSourcePtr())
-	{
-		return MediaPlateResource->GetResourceType();
-	}
-
-	return DefaultType;
-}
-
-FMediaPlateResource* FMediaPlateResourceCustomization::GetMediaPlateSourcePtr() const
-{
-	if (!MediaPlateResourcePropertyHandle)
-	{
-		return nullptr;
-	}
-
-	void* StructRawData;
-	const FPropertyAccess::Result AccessResult = MediaPlateResourcePropertyHandle->GetValueData(StructRawData);
-	if (AccessResult == FPropertyAccess::Success)
-	{
-		return static_cast<FMediaPlateResource*>(StructRawData);
-	}
-
-	return nullptr;
-}
-
-UObject* FMediaPlateResourceCustomization::GetMediaPlateResourceOwner() const
-{
-	TArray<UObject*> OuterObjects;
-	MediaPlateResourcePropertyHandle->GetOuterObjects(OuterObjects);
-
-	if (OuterObjects.Num() != 1)
-	{
-		return nullptr;
-	}
-
-	if (UObject* Outer = OuterObjects[0])
-	{
-		return Outer;
-	}
-
-	return nullptr;
-}
-
-void FMediaPlateResourceCustomization::OnAssetTypeChanged(EMediaPlateResourceType InMediaSourceType)
+TOptional<EMediaPlateResourceType> FMediaPlateResourceCustomization::GetAssetType() const
 {
-	if (GetAssetType() != InMediaSourceType)
+	if (ResourceTypePropertyHandle)
 	{
-		// Updating value and notifying, in case the user is switching between already specified options
-		if (FMediaPlateResource* MediaPlateResource = GetMediaPlateSourcePtr())
+		using UnderlyingType = std::underlying_type_t<EMediaPlateResourceType>;
+		UnderlyingType ResourceType;
+		if (ResourceTypePropertyHandle->GetValue(ResourceType) == FPropertyAccess::Success)
 		{
-			const FScopedTransaction Transaction(LOCTEXT("OnMediaSourceTypeChanged", "Media Source type changed"));
-
-			if (UObject* MediaPlateResourceOwner = GetMediaPlateResourceOwner())
-			{
-				MediaPlateResourceOwner->Modify();
-				MediaPlateResource->Modify();
-			}
-
-			MediaPlateResource->SetResourceType(InMediaSourceType);
-			MediaPlateResourcePropertyHandle->NotifyPostChange(EPropertyChangeType::ValueSet);
+			return static_cast<EMediaPlateResourceType>(ResourceType);
 		}
 	}
+	return TOptional<EMediaPlateResourceType>();
 }
 
-FString FMediaPlateResourceCustomization::GetMediaAssetPath() const
-{
-	FString MediaAssetPath;
-
-	if (const FMediaPlateResource* MediaPlateResource = GetMediaPlateSourcePtr())
-	{
-		return MediaPlateResource->GetMediaAsset()->GetPathName();
-	}
-
-	return MediaAssetPath;
-}
-
-FString FMediaPlateResourceCustomization::GetPlaylistPath() const
+void FMediaPlateResourceCustomization::OnAssetTypeChanged(TOptional<EMediaPlateResourceType> InMediaSourceType)
 {
-	FString PlaylistAssetPath;
-
-	if (const FMediaPlateResource* MediaPlateResource = GetMediaPlateSourcePtr())
+	if (ResourceTypePropertyHandle && InMediaSourceType.IsSet())
 	{
-		if (const UMediaPlaylist* Playlist = MediaPlateResource->GetSourcePlaylist())
-		{
-			PlaylistAssetPath = Playlist->GetPathName();
-		}
+		using UnderType = std::underlying_type_t<EMediaPlateResourceType>;
+		ResourceTypePropertyHandle->SetValue(*reinterpret_cast<UnderType*>(&InMediaSourceType));
 	}
-
-	return PlaylistAssetPath;
 }
 
 FString FMediaPlateResourceCustomization::GetMediaPath() const
 {
-	FString ExternalMediaPath;
-
-	if (const FMediaPlateResource* MediaPlateResource = GetMediaPlateSourcePtr())
-	{
-		ExternalMediaPath = MediaPlateResource->GetExternalMediaPath();
-	}
-
-	return ExternalMediaPath;
-}
-
-void FMediaPlateResourceCustomization::OnMediaAssetChanged(const FAssetData& InAssetData)
-{
-	if (!MediaPlateResourcePropertyHandle)
+	if (ExternalMediaPathPropertyHandle)
 	{
-		return;
-	}
-
-	if (const UMediaSource* MediaAsset = Cast<UMediaSource>(InAssetData.GetAsset()))
-	{
-		if (FMediaPlateResource* MediaPlateResource = GetMediaPlateSourcePtr())
+		FString Path;
+		const FPropertyAccess::Result AccessResult = ExternalMediaPathPropertyHandle->GetValue(Path);
+		if (AccessResult == FPropertyAccess::Success)
 		{
-			if (UObject* MediaPlateResourceOwner = GetMediaPlateResourceOwner())
-			{
-				const FScopedTransaction Transaction(LOCTEXT("OnMediaSourceAssetChanged", "Media Source asset changed"));
-				MediaPlateResourceOwner->Modify();
-				MediaPlateResource->Modify();
-
-				MediaPlateResource->SelectAsset(MediaAsset, MediaPlateResourceOwner);
-				MediaPlateResourcePropertyHandle->NotifyPostChange(EPropertyChangeType::ValueSet);
-			}
+			return Path;
 		}
-	}
-}
-
-void FMediaPlateResourceCustomization::OnMediaPathPicked(const FString& InPickedPath)
-{
-	if (!MediaPlateResourcePropertyHandle)
-	{
-		return;
-	}
-
-	FMediaPlateResource* MediaPlateResource = GetMediaPlateSourcePtr();
-
-	if (!MediaPlateResource)
-	{
-		return;
-	}
-
-	if (!InPickedPath.IsEmpty() && InPickedPath != MediaPlateResource->GetExternalMediaPath())
-	{
-		const FScopedTransaction Transaction(LOCTEXT("OnMediaExternalPathChanged", "Media external file path changed"));
-		if (UObject* MediaPlateResourceOwner = GetMediaPlateResourceOwner())
+		if (AccessResult == FPropertyAccess::MultipleValues)
 		{
-			MediaPlateResourceOwner->Modify();
-			MediaPlateResource->Modify();
-
-			MediaPlateResource->LoadExternalMedia(InPickedPath, MediaPlateResourceOwner);
-			MediaPlateResourcePropertyHandle->NotifyPostChange(EPropertyChangeType::ValueSet);
+			return TEXT("(Multiple values)");
 		}
 	}
+	return FString();
 }
 
-void FMediaPlateResourceCustomization::OnPlaylistChanged(const FAssetData& InAssetData)
+void FMediaPlateResourceCustomization::OnMediaPathPicked(const FString& InPickedPath)
 {
-	if (!MediaPlateResourcePropertyHandle)
+	if (ExternalMediaPathPropertyHandle)
 	{
-		return;
-	}
-
-	const UMediaPlaylist* Playlist = Cast<UMediaPlaylist>(InAssetData.GetAsset());
-	if (FMediaPlateResource* MediaPlateResource = GetMediaPlateSourcePtr())
-	{
-		const FScopedTransaction Transaction(LOCTEXT("OnPlaylistChanged", "Media Playlist changed"));
-		if (UObject* MediaPlateResourceOwner = GetMediaPlateResourceOwner())
-		{
-			MediaPlateResourceOwner->Modify();
-		}
-
-		// No need to call MediaPlateResource->Modify():
-		// ActivePlaylist property will be overwritten, not modified
-
-		MediaPlateResource->SelectPlaylist(Playlist);
-		MediaPlateResourcePropertyHandle->NotifyPostChange(EPropertyChangeType::ValueSet);
+		ExternalMediaPathPropertyHandle->SetValue(InPickedPath);
 	}
 }
 
@@ -331,4 +208,9 @@ EVisibility FMediaPlateResourceCustomization::GetPlaylistSelectorVisibility() co
 	return GetAssetType() == EMediaPlateResourceType::Playlist ? EVisibility::Visible : EVisibility::Collapsed;
 }
 
+EVisibility FMediaPlateResourceCustomization::GetMultipleValuesVisibility() const
+{
+	return GetAssetType().IsSet() ? EVisibility::Collapsed : EVisibility::Visible;
+}
+
 #undef LOCTEXT_NAMESPACE
diff --git a/Engine/Plugins/Media/MediaPlate/Source/MediaPlateEditor/Private/MediaPlateResourceCustomization.h b/Engine/Plugins/Media/MediaPlate/Source/MediaPlateEditor/Private/MediaPlateResourceCustomization.h
--- a/Engine/Plugins/Media/MediaPlate/Source/MediaPlateEditor/Private/MediaPlateResourceCustomization.h
+++ b/Engine/Plugins/Media/MediaPlate/Source/MediaPlateEditor/Private/MediaPlateResourceCustomization.h
@@ -6,7 +6,6 @@
 #include "MediaPlateComponent.h"
 
 class IPropertyHandle;
-class UMediaSource;
 struct EVisibility;
 struct FMediaPlateResource;
 
@@ -26,22 +25,22 @@ public:
 	//~ End IPropertyTypeCustomization
 
 private:
-	EMediaPlateResourceType GetAssetType() const;
+	TOptional<EMediaPlateResourceType> GetAssetType() const;
 
-	FString GetMediaAssetPath() const;
-	FString GetPlaylistPath() const;
 	FString GetMediaPath() const;
 
-	void OnAssetTypeChanged(EMediaPlateResourceType InMediaSourceType);
-	void OnPlaylistChanged(const FAssetData& InAssetData);
-	void OnMediaAssetChanged(const FAssetData& InAssetData);
+	void OnAssetTypeChanged(TOptional<EMediaPlateResourceType> InResourceType);
 	void OnMediaPathPicked(const FString& InPickedPath);
 
 	EVisibility GetAssetSelectorVisibility() const;
 	EVisibility GetFileSelectorVisibility() const;
 	EVisibility GetPlaylistSelectorVisibility() const;
-	FMediaPlateResource* GetMediaPlateSourcePtr() const;
-	UObject* GetMediaPlateResourceOwner() const;
+	EVisibility GetMultipleValuesVisibility() const;
 
+
+	TSharedPtr<IPropertyHandle> ResourceTypePropertyHandle;
 	TSharedPtr<IPropertyHandle> MediaPlateResourcePropertyHandle;
+	TSharedPtr<IPropertyHandle> ExternalMediaPathPropertyHandle;
+	TSharedPtr<IPropertyHandle> MediaAssetPropertyHandle;
+	TSharedPtr<IPropertyHandle> SourcePlaylistPropertyHandle;
 };
