[MediaPlate] Only restart the media player (on media plate resource change) if it the media plate is "playing" (i.e. open has been called).

The "IntendedPlaybackState" is not in a reliable state if the media plate is not "playing" (i.e. open has been called).

diff --git a/Engine/Plugins/Media/MediaPlate/Source/MediaPlate/Private/MediaPlateComponent.cpp b/Engine/Plugins/Media/MediaPlate/Source/MediaPlate/Private/MediaPlateComponent.cpp
--- a/Engine/Plugins/Media/MediaPlate/Source/MediaPlate/Private/MediaPlateComponent.cpp
+++ b/Engine/Plugins/Media/MediaPlate/Source/MediaPlate/Private/MediaPlateComponent.cpp
@@ -1035,7 +1035,7 @@ bool UMediaPlateComponent::RestartPlayer()
 {
 	if (MediaPlayer != nullptr)
 	{
-		if (IntendedPlaybackState == EPlaybackState::Playing)
+		if (bIsMediaPlatePlaying && IntendedPlaybackState == EPlaybackState::Playing)
 		{
 			MediaPlayer->Close();
 			Open();
