[MediaPlate] Opening and playing a media plate will dirty the level package in editor.

The automatic aspect ratio adjustment will systematically lead to the package being marked as dirty.
To prevent this, we check if the scale has changed. If it hasn't changed, there shouldn't be a need to mark it dirty.

diff --git a/Engine/Plugins/Media/MediaPlate/Source/MediaPlate/Private/MediaPlateComponent.cpp b/Engine/Plugins/Media/MediaPlate/Source/MediaPlate/Private/MediaPlateComponent.cpp
--- a/Engine/Plugins/Media/MediaPlate/Source/MediaPlate/Private/MediaPlateComponent.cpp
+++ b/Engine/Plugins/Media/MediaPlate/Source/MediaPlate/Private/MediaPlateComponent.cpp
@@ -758,9 +758,13 @@ void UMediaPlateComponent::SetAspectRatio(float AspectRatio)
 		{
 			Height = 1.0f / AspectRatio;
 		}
-		FVector Scale(1.0f, 1.0f, Height);
+		const FVector Scale(1.0f, 1.0f, Height);
 #if WITH_EDITOR
-		StaticMeshComponent->Modify();
+		// Avoid dirtying the package if scale is not modified.
+		if (Scale != StaticMeshComponent->GetRelativeScale3D())
+		{
+			StaticMeshComponent->Modify();
+		}
 #endif // WITH_EDITOR
 		StaticMeshComponent->SetRelativeScale3D(Scale);
 
