[MediaPlate] Adding legacy support for MediaPlate with img sequence in "external file" mode.

For legacy code, we want the extracted media path to be valid for use with SpawnMediaSourceForString. In the case of img sequence, it would need to be a filepath with an extension. GetUrl returns a directory, therefore additional code is needed to find a file compatible in that directory.

Also, the new function will try to make the path relative to project.

diff --git a/Engine/Plugins/Media/MediaPlate/Source/MediaPlate/Private/MediaPlateComponent.cpp b/Engine/Plugins/Media/MediaPlate/Source/MediaPlate/Private/MediaPlateComponent.cpp
--- a/Engine/Plugins/Media/MediaPlate/Source/MediaPlate/Private/MediaPlateComponent.cpp
+++ b/Engine/Plugins/Media/MediaPlate/Source/MediaPlate/Private/MediaPlateComponent.cpp
@@ -5,6 +5,7 @@
 #include "Components/StaticMeshComponent.h"
 #include "Engine/Engine.h"
 #include "Engine/StaticMesh.h"
+#include "HAL/FileManager.h"
 #include "IMediaAssetsModule.h"
 #include "IMediaClock.h"
 #include "IMediaClockSink.h"
@@ -72,6 +73,69 @@ namespace UE::MediaPlateComponent
 	{
 		return InPlaylist ? InPlaylist->Get(InPlaylistIndex) : nullptr;
 	}
+
+#if WITH_EDITOR
+	/**
+	 * Legacy support code to extract a media path compatible with SpawnMediaSourceForString from the
+	 * existing MediaSource.
+	 * 
+	 * @param InExternalMediaSource Existing media source we want to extract a path from.
+	 * @param InOuter Outer object for the temporary creation of a media source to validate.
+	 * @return Extracted path.
+	 */
+	FString ExtractExternalMediaPath(const UMediaSource* InExternalMediaSource, UObject* InOuter)
+	{
+		FString MediaPath = InExternalMediaSource->GetUrl();
+
+		const FString FilePrefix(TEXT("file://"));
+		const FString ImgPrefix(TEXT("img://"));
+		if (MediaPath.StartsWith(FilePrefix))
+		{
+			MediaPath = MediaPath.RightChop(FilePrefix.Len());
+		}
+		else if (MediaPath.StartsWith(ImgPrefix))
+		{
+			// Special case of img sequence.
+			// We need to extract a path that is compatible with SpawnMediaSourceForString.
+			// As it currently stands, this requires a valid file path with extension.
+			
+			MediaPath = MediaPath.RightChop(ImgPrefix.Len());
+
+			// This is going to be a directory name
+			if (FPaths::DirectoryExists(MediaPath))
+			{
+				// Get the first compatible file.
+				TArray<FString> FoundFiles;
+				IFileManager::Get().FindFiles(FoundFiles, *MediaPath, TEXT("*"));
+				for (const FString& File : FoundFiles)
+				{
+					const FString Fullpath = FPaths::Combine(MediaPath, File);
+					if (UMediaSource::SpawnMediaSourceForString(Fullpath, InOuter))
+					{
+						MediaPath = Fullpath;
+						break;
+					}
+				}
+			}
+		}
+
+		// At least support the case if the path is within the project (otherwise, leave absolute)
+		FString RelativeToProject = MediaPath;
+		if (FPaths::MakePathRelativeTo(RelativeToProject, *FPaths::ProjectDir()))
+		{
+			if (!RelativeToProject.StartsWith(TEXT(".")))
+			{
+				// Project relative paths are expected to start with "./"
+				MediaPath = FPaths::Combine(TEXT("."), RelativeToProject);
+			}
+			else
+			{
+				MediaPath = RelativeToProject;
+			}
+		}
+		return MediaPath;
+	}
+#endif
 };
 
 /**
@@ -154,6 +218,16 @@ PRAGMA_DISABLE_DEPRECATION_WARNINGS
 	{
 		InitializeMediaPlateResource();
 	}
+
+	if(MediaPlateResource.ExternalMedia_DEPRECATED)
+	{
+		// Move the existing media source.
+		ExternalMediaSource = MediaPlateResource.ExternalMedia_DEPRECATED;
+		MediaPlateResource.ExternalMedia_DEPRECATED = nullptr;
+		
+		// Update external media path to match (as best we can).
+		MediaPlateResource.LoadExternalMedia(UE::MediaPlateComponent::ExtractExternalMediaPath(ExternalMediaSource, this));
+	}
 PRAGMA_ENABLE_DEPRECATION_WARNINGS
 	
 	// Ensure the external media source is loaded.
@@ -1460,25 +1534,14 @@ void UMediaPlateComponent::InitializeMediaPlateResource()
 	}
 	else
 	{
-		if (const UMediaSource* MediaSource = Playlist->Get(0))
+		if (UMediaSource* MediaSource = Playlist->Get(0))
 		{
 			// If outer is this, source has been created from external file
 			if (MediaSource->GetOuter() == this)
 			{
-				FString MediaPath = MediaSource->GetUrl();
-
-				const FString FilePrefix(TEXT("file://"));
-				const FString ImgPrefix(TEXT("img://"));
-				if (MediaPath.StartsWith(FilePrefix))
-				{
-					MediaPath = MediaPath.RightChop(FilePrefix.Len());
-				}
-				else if (MediaPath.StartsWith(ImgPrefix))
-				{
-					MediaPath = MediaPath.RightChop(ImgPrefix.Len());
-				}
-
-				SelectExternalMedia(MediaPath);
+				ExternalMediaSource = MediaSource;
+				// Update external media path to match (as best we can).
+				MediaPlateResource.LoadExternalMedia(UE::MediaPlateComponent::ExtractExternalMediaPath(ExternalMediaSource, this));
 			}
 			else
 			{
diff --git a/Engine/Plugins/Media/MediaPlate/Source/MediaPlate/Public/MediaPlateResource.h b/Engine/Plugins/Media/MediaPlate/Source/MediaPlate/Public/MediaPlateResource.h
--- a/Engine/Plugins/Media/MediaPlate/Source/MediaPlate/Public/MediaPlateResource.h
+++ b/Engine/Plugins/Media/MediaPlate/Source/MediaPlate/Public/MediaPlateResource.h
@@ -96,6 +96,14 @@ public:
 	void Modify() const {}
 #endif	
 
+PRAGMA_DISABLE_DEPRECATION_WARNINGS
+	FMediaPlateResource() = default;
+	FMediaPlateResource(const FMediaPlateResource& InOther) = default;
+	FMediaPlateResource(FMediaPlateResource&& InOther) = default;
+	FMediaPlateResource& operator=(const FMediaPlateResource&) = default;
+	FMediaPlateResource& operator=(FMediaPlateResource&&) = default;
+PRAGMA_ENABLE_DEPRECATION_WARNINGS
+
 private:
 friend class UMediaPlateComponent;
 friend class FMediaPlateResourceCustomization;
@@ -119,4 +127,10 @@ friend class FMediaPlateResourceCustomization;
 	/** User facing Playlist asset */
 	UPROPERTY(EditAnywhere, Category = "MediaPlateResource")
 	TSoftObjectPtr<UMediaPlaylist> SourcePlaylist;
+
+#if WITH_EDITORONLY_DATA
+	UE_DEPRECATED(5.6, "Moved to UMediaPlateComponent::ExternalMediaSource")
+	UPROPERTY(Instanced, meta=(DeprecatedProperty, DeprecationMessage="Moved to UMediaPlateComponent::ExternalMediaSource"))
+	TObjectPtr<UMediaSource> ExternalMedia_DEPRECATED;
+#endif
 };
\ No newline at end of file
