[MediaPlayerEditor] Fix Media IO - Player info not showing on the Media Details Panel while playing

diff --git a/Engine/Plugins/Media/MediaPlate/Source/MediaPlateEditor/Private/Widgets/SMediaPlateEditorMediaDetails.cpp b/Engine/Plugins/Media/MediaPlate/Source/MediaPlateEditor/Private/Widgets/SMediaPlateEditorMediaDetails.cpp
--- a/Engine/Plugins/Media/MediaPlate/Source/MediaPlateEditor/Private/Widgets/SMediaPlateEditorMediaDetails.cpp
+++ b/Engine/Plugins/Media/MediaPlate/Source/MediaPlateEditor/Private/Widgets/SMediaPlateEditorMediaDetails.cpp
@@ -132,6 +132,24 @@ void SMediaPlateEditorMediaDetails::Construct(const FArguments& InArgs,
 									SAssignNew(NumTilesText, STextBlock)
 									.TextStyle(SlateStyle, StyleName)
 								]
+							// Start Timecode.
+							+ SVerticalBox::Slot()
+								.AutoHeight()
+								.VAlign(VAlign_Center)
+								.Padding(4.0f)
+								[
+									SAssignNew(StartTimecodeText, STextBlock)
+									.TextStyle(SlateStyle, StyleName)
+								]
+							// Seek Performance
+							+ SVerticalBox::Slot()
+								.AutoHeight()
+								.VAlign(VAlign_Center)
+								.Padding(4.0f)
+								[
+									SAssignNew(SeekPerformance, STextBlock)
+									.TextStyle(SlateStyle, StyleName)
+								]
 						]
 				]
 
@@ -161,6 +179,8 @@ void SMediaPlateEditorMediaDetails::UpdateDetails()
 	int64 ResourceSize = 0;
 	int32 SurfaceWidth = 0;
 	int32 SurfaceHeight = 0;
+	FString StartTimecode;
+	int32 KeyframeInterval = -1;
 
 	if (MediaPlate.IsValid())
 	{
@@ -174,6 +194,8 @@ void SMediaPlateEditorMediaDetails::UpdateDetails()
 			FIntPoint NumTiles(EForceInit::ForceInitToZero);
 			MediaPlayer->GetMediaInfo<FIntPoint>(NumTiles, UMediaPlayer::MediaInfoNameSourceNumTiles.Resolve());
 			NumTotalTiles = NumTiles.X * NumTiles.Y;
+			MediaPlayer->GetMediaInfo<FString>(StartTimecode, UMediaPlayer::MediaInfoNameStartTimecodeValue.Resolve());
+			MediaPlayer->GetMediaInfo<int32>(KeyframeInterval, UMediaPlayer::MediaInfoNameKeyframeInterval.Resolve());
 		}
 
 		// Get texture info,
@@ -210,6 +232,30 @@ void SMediaPlateEditorMediaDetails::UpdateDetails()
 		FText::AsNumber(SurfaceWidth), FText::AsNumber(SurfaceHeight)));
 	ResourceSizeText->SetText(FText::Format(LOCTEXT("ResourceSize", "Resource Size: {0} KB"),
 		FText::AsNumber(ResourceSize)));
+
+	if (!StartTimecode.IsEmpty())
+	{
+		StartTimecodeText->SetText(FText::Format(LOCTEXT("StartTimecode", "Start Timecode: {0}"), FText::FromString(StartTimecode)));
+	}
+	StartTimecodeText->SetVisibility(!StartTimecode.IsEmpty() ? EVisibility::Visible : EVisibility::Collapsed);
+	
+	/*
+		The keyframe interval is one of:
+		  -1 : no information returned (unknown)
+		   0 : unknown keyframe spacing, not every frame is a keyframe but the spacing is variable or cannot be determined
+		   1 : every frame is a keyframe
+		  >1 : every n'th frame is a keyframe
+	*/
+	if (KeyframeInterval >= 0)
+	{
+		SeekPerformance->SetText(FText::Format(LOCTEXT("SeekPerformance", "Seek Performance: {0}"), KeyframeInterval == 1 ? LOCTEXT("SeekPerformanceF", "Fast") : LOCTEXT("SeekPerformanceS", "Slow (GOP codec)")));
+		SeekPerformance->SetVisibility(EVisibility::Visible);
+	}
+	else
+	{
+		SeekPerformance->SetVisibility(EVisibility::Collapsed);
+	}
+
 }
 
 #undef LOCTEXT_NAMESPACE
diff --git a/Engine/Plugins/Media/MediaPlate/Source/MediaPlateEditor/Private/Widgets/SMediaPlateEditorMediaDetails.h b/Engine/Plugins/Media/MediaPlate/Source/MediaPlateEditor/Private/Widgets/SMediaPlateEditorMediaDetails.h
--- a/Engine/Plugins/Media/MediaPlate/Source/MediaPlateEditor/Private/Widgets/SMediaPlateEditorMediaDetails.h
+++ b/Engine/Plugins/Media/MediaPlate/Source/MediaPlateEditor/Private/Widgets/SMediaPlateEditorMediaDetails.h
@@ -54,4 +54,6 @@ private:
 	TSharedPtr<STextBlock> NumTilesText;
 	TSharedPtr<STextBlock> ResolutionText;
 	TSharedPtr<STextBlock> ResourceSizeText;
+	TSharedPtr<STextBlock> StartTimecodeText;
+	TSharedPtr<STextBlock> SeekPerformance;
 };
diff --git a/Engine/Plugins/Media/MediaPlayerEditor/Source/MediaPlayerEditor/Private/Widgets/SMediaPlayerEditorMediaDetails.cpp b/Engine/Plugins/Media/MediaPlayerEditor/Source/MediaPlayerEditor/Private/Widgets/SMediaPlayerEditorMediaDetails.cpp
--- a/Engine/Plugins/Media/MediaPlayerEditor/Source/MediaPlayerEditor/Private/Widgets/SMediaPlayerEditorMediaDetails.cpp
+++ b/Engine/Plugins/Media/MediaPlayerEditor/Source/MediaPlayerEditor/Private/Widgets/SMediaPlayerEditorMediaDetails.cpp
@@ -34,6 +34,15 @@ void SMediaPlayerEditorMediaDetails::Construct(const FArguments& InArgs,
 						[
 							SNew(SVerticalBox)
 
+							// Player name.
+							+ SVerticalBox::Slot()
+								.AutoHeight()
+								.VAlign(VAlign_Center)
+								.Padding(4.0f)
+								[
+									SAssignNew(MediaPlayerName, STextBlock)
+								]
+
 							// Resolution.
 							+ SVerticalBox::Slot()
 								.AutoHeight()
@@ -146,6 +155,7 @@ void SMediaPlayerEditorMediaDetails::Tick(const FGeometry& AllottedGeometry, con
 
 void SMediaPlayerEditorMediaDetails::UpdateDetails()
 {
+	FName PlayerName;
 	FString Format;
 	float FrameRate = 0.0f;
 	int32 LODBias = 0;
@@ -161,6 +171,7 @@ void SMediaPlayerEditorMediaDetails::UpdateDetails()
 	// Get player info.
 	if (MediaPlayer != nullptr)
 	{
+		PlayerName = MediaPlayer->GetPlayerName();
 		FrameRate = MediaPlayer->GetVideoTrackFrameRate(INDEX_NONE, INDEX_NONE);
 		Format = MediaPlayer->GetVideoTrackType(INDEX_NONE, INDEX_NONE);
 		FIntPoint NumTiles(EForceInit::ForceInitToZero);
@@ -185,6 +196,8 @@ void SMediaPlayerEditorMediaDetails::UpdateDetails()
 	}
 
 	// Update text.
+	MediaPlayerName->SetText(FText::Format(LOCTEXT("Player", "Player: {0}"),
+		FText::FromName(PlayerName)));
 	FormatText->SetText(FText::Format(LOCTEXT("Format", "Format: {0}"),
 		FText::FromString(Format)));
 	FrameRateText->SetText(FText::Format(LOCTEXT("FrameRate", "Frame Rate: {0}"),
diff --git a/Engine/Plugins/Media/MediaPlayerEditor/Source/MediaPlayerEditor/Public/Widgets/SMediaPlayerEditorMediaDetails.h b/Engine/Plugins/Media/MediaPlayerEditor/Source/MediaPlayerEditor/Public/Widgets/SMediaPlayerEditorMediaDetails.h
--- a/Engine/Plugins/Media/MediaPlayerEditor/Source/MediaPlayerEditor/Public/Widgets/SMediaPlayerEditorMediaDetails.h
+++ b/Engine/Plugins/Media/MediaPlayerEditor/Source/MediaPlayerEditor/Public/Widgets/SMediaPlayerEditorMediaDetails.h
@@ -49,6 +49,7 @@ private:
 	TWeakObjectPtr<UMediaTexture> MediaTexture;
 
 	/** Our widgets. */
+	TSharedPtr<STextBlock> MediaPlayerName;
 	TSharedPtr<STextBlock> FormatText;
 	TSharedPtr<STextBlock> FrameRateText;
 	TSharedPtr<STextBlock> LODBiasText;
