MediaPlate: Prompt user to delete deprecated overlay material when enabling holdout composite.

diff --git a/Engine/Plugins/Media/MediaPlate/Source/MediaPlate/Private/MediaPlate.cpp b/Engine/Plugins/Media/MediaPlate/Source/MediaPlate/Private/MediaPlate.cpp
--- a/Engine/Plugins/Media/MediaPlate/Source/MediaPlate/Private/MediaPlate.cpp
+++ b/Engine/Plugins/Media/MediaPlate/Source/MediaPlate/Private/MediaPlate.cpp
@@ -19,6 +19,7 @@
 #if WITH_EDITOR
 #include "Editor.h"
 #include "MediaPlateAssetUserData.h"
+#include "Misc/MessageDialog.h"
 #include "UObject/ObjectSaveContext.h"
 #endif
 
@@ -49,6 +50,20 @@ namespace UE::MediaPlate::Private
 			TranslucencySPBasisCVar->Set(InBasis);
 		}
 	}
+
+	/** Check for the presence of the (now deprecated) overlay composite material. */
+	bool HasDeprecatedOverlayCompositeMaterial(AMediaPlate& InMediaPlate)
+	{
+		if (UMaterialInterface* OverlayMaterial = InMediaPlate.GetCurrentOverlayMaterial())
+		{
+			if (UMaterial* ParentMaterial = OverlayMaterial->GetMaterial())
+			{
+				return ParentMaterial->GetPathName() == TEXT("/MediaPlate/M_MediaPlate_OverlayComp.M_MediaPlate_OverlayComp");
+			}
+		}
+
+		return false;
+	}
 }
 
 AMediaPlate::AMediaPlate(const FObjectInitializer& ObjectInitializer)
@@ -111,6 +126,8 @@ void AMediaPlate::PostLoad()
 
 				HoldoutCompositeComponent->DestroyComponent();
 			}
+
+			//TODO (5.7): Remove deprecated overlay composite material & enable holdout composite.
 		}
 	}
 #endif
@@ -227,6 +244,32 @@ void AMediaPlate::PostEditChangeProperty(FPropertyChangedEvent& PropertyChangedE
 	if (PropertyChangedEvent.GetPropertyName() == GET_MEMBER_NAME_CHECKED(ThisClass, bEnableHoldoutComposite))
 	{
 		SetHoldoutCompositeEnabled(bEnableHoldoutComposite);
+		
+		// TODO (5.7): Remove once we fully migrate to holdout composite.
+		if (bEnableHoldoutComposite && UE::MediaPlate::Private::HasDeprecatedOverlayCompositeMaterial(*this))
+		{
+			static TOptional<EAppReturnType::Type> bYesAllNoAll;
+			EAppReturnType::Type Response = EAppReturnType::No;
+
+			if (bYesAllNoAll.IsSet())
+			{
+				Response = bYesAllNoAll.GetValue();
+			}
+			else
+			{
+				Response = FMessageDialog::Open(EAppMsgType::YesNoYesAllNoAll, LOCTEXT("RemoveOverlayMaterial", "The deprecated overlay composite material conflicts with holdout composite.\nProceed to delete it?"));
+
+				if (Response == EAppReturnType::YesAll || Response == EAppReturnType::NoAll)
+				{
+					bYesAllNoAll = Response;
+				}
+			}
+
+			if (Response == EAppReturnType::Yes || Response == EAppReturnType::YesAll)
+			{
+				StaticMeshComponent->SetOverlayMaterial(nullptr);
+			}
+		}
 	}
 }
 
