[MediaPlate] Exposing Next() and Previous() to blueprint.

diff --git a/Engine/Plugins/Media/MediaPlate/Source/MediaPlate/Public/MediaPlateComponent.h b/Engine/Plugins/Media/MediaPlate/Source/MediaPlate/Public/MediaPlateComponent.h
--- a/Engine/Plugins/Media/MediaPlate/Source/MediaPlate/Public/MediaPlateComponent.h
+++ b/Engine/Plugins/Media/MediaPlate/Source/MediaPlate/Public/MediaPlateComponent.h
@@ -139,6 +139,7 @@ public:
 	 *
 	 * returns	True if it played something.
 	 */
+	UFUNCTION(BlueprintCallable, Category = "Media|MediaPlateComponent")
 	UE_API bool Next();
 
 	/**
@@ -160,6 +161,7 @@ public:
 	 *
 	 * returns	True if it played something.
 	 */
+	UFUNCTION(BlueprintCallable, Category = "Media|MediaPlateComponent")
 	UE_API bool Previous();
 
 	/**
