MediaPlate: Full deprecation of overlay material composite technique.

diff --git a/Engine/Plugins/Media/MediaPlate/Source/MediaPlate/Private/MediaPlate.cpp b/Engine/Plugins/Media/MediaPlate/Source/MediaPlate/Private/MediaPlate.cpp
--- a/Engine/Plugins/Media/MediaPlate/Source/MediaPlate/Private/MediaPlate.cpp
+++ b/Engine/Plugins/Media/MediaPlate/Source/MediaPlate/Private/MediaPlate.cpp
@@ -32,25 +32,6 @@ FLazyName AMediaPlate::MediaTextureName("MediaTexture");
 
 namespace UE::MediaPlate::Private
 {
-	// TODO: Could be removed (with associated materials) now that we've deprecated overlay compositing.
-	void ApplyTranslucencyScreenPercentageCVar(int32 InBasis)
-	{
-		static IConsoleVariable* TranslucencySPBasisCVar = IConsoleManager::Get().FindConsoleVariable(TEXT("r.Translucency.ScreenPercentage.Basis"));
-		if (TranslucencySPBasisCVar->GetInt() != InBasis)
-		{
-			if (InBasis)
-			{
-				UE_LOG(LogMediaPlate, Warning, TEXT("Setting 'r.Translucency.ScreenPercentage.Basis' to 1. For media plates with overlay materials, please apply this console variable permanently to your project."));
-			}
-			else
-			{
-				UE_LOG(LogMediaPlate, Warning, TEXT("Setting 'r.Translucency.ScreenPercentage.Basis' to 0."));
-			}
-
-			TranslucencySPBasisCVar->Set(InBasis);
-		}
-	}
-
 	/** Check for the presence of the (now deprecated) overlay composite material. */
 	bool HasDeprecatedOverlayCompositeMaterial(AMediaPlate& InMediaPlate)
 	{
@@ -115,19 +96,26 @@ void AMediaPlate::PostLoad()
 	Super::PostLoad();
 
 #if WITH_EDITORONLY_DATA
-	if (!HasAnyFlags(RF_ClassDefaultObject))
+	// The holdout composite component was replaced with a checkbox.
+	if (GetLinkerCustomVersion(FFortniteMainBranchObjectVersion::GUID) < FFortniteMainBranchObjectVersion::MediaPlateHoldoutComponentRemoval)
 	{
-		// The holdout composite component was replaced with a checkbox.
-		if (GetLinkerCustomVersion(FFortniteMainBranchObjectVersion::GUID) < FFortniteMainBranchObjectVersion::MediaPlateHoldoutComponentRemoval)
+		if (UHoldoutCompositeComponent* HoldoutCompositeComponent = FindComponentByClass<UHoldoutCompositeComponent>())
 		{
-			if (UHoldoutCompositeComponent* HoldoutCompositeComponent = FindComponentByClass<UHoldoutCompositeComponent>())
-			{
-				bEnableHoldoutComposite = HoldoutCompositeComponent->IsEnabled();
+			bEnableHoldoutComposite = HoldoutCompositeComponent->IsEnabled();
 
-				HoldoutCompositeComponent->DestroyComponent();
-			}
+			HoldoutCompositeComponent->DestroyComponent();
+		}
+	}
+
+	if (GetLinkerCustomVersion(FFortniteMainBranchObjectVersion::GUID) < FFortniteMainBranchObjectVersion::MediaPlateOverlayTechniqueRemoval)
+	{
+		if (UE::MediaPlate::Private::HasDeprecatedOverlayCompositeMaterial(*this))
+		{
+			UE_LOG(LogMediaPlate, Display, TEXT("Deprecated overlay material technique automatically replaced with holdout composite."));
 
-			//TODO (5.7): Remove deprecated overlay composite material & enable holdout composite.
+			StaticMeshComponent->SetOverlayMaterial(nullptr);
+
+			bEnableHoldoutComposite = true;
 		}
 	}
 #endif
@@ -168,11 +156,6 @@ void AMediaPlate::PostRegisterAllComponents()
 
 	if (IsValid(StaticMeshComponent))
 	{
-		if (IsValid(StaticMeshComponent->OverlayMaterial))
-		{
-			UE::MediaPlate::Private::ApplyTranslucencyScreenPercentageCVar(1);
-		}
-
 		ConditionallyEnableHoldoutComposite();
 	}
 }
@@ -244,32 +227,6 @@ void AMediaPlate::PostEditChangeProperty(FPropertyChangedEvent& PropertyChangedE
 	if (PropertyChangedEvent.GetPropertyName() == GET_MEMBER_NAME_CHECKED(ThisClass, bEnableHoldoutComposite))
 	{
 		SetHoldoutCompositeEnabled(bEnableHoldoutComposite);
-		
-		// TODO (5.7): Remove once we fully migrate to holdout composite.
-		if (bEnableHoldoutComposite && UE::MediaPlate::Private::HasDeprecatedOverlayCompositeMaterial(*this))
-		{
-			static TOptional<EAppReturnType::Type> bYesAllNoAll;
-			EAppReturnType::Type Response = EAppReturnType::No;
-
-			if (bYesAllNoAll.IsSet())
-			{
-				Response = bYesAllNoAll.GetValue();
-			}
-			else
-			{
-				Response = FMessageDialog::Open(EAppMsgType::YesNoYesAllNoAll, LOCTEXT("RemoveOverlayMaterial", "The deprecated overlay composite material conflicts with holdout composite.\nProceed to delete it?"));
-
-				if (Response == EAppReturnType::YesAll || Response == EAppReturnType::NoAll)
-				{
-					bYesAllNoAll = Response;
-				}
-			}
-
-			if (Response == EAppReturnType::Yes || Response == EAppReturnType::YesAll)
-			{
-				StaticMeshComponent->SetOverlayMaterial(nullptr);
-			}
-		}
 	}
 }
 
@@ -441,8 +398,6 @@ void AMediaPlate::ApplyOverlayMaterial(UMaterialInterface* InOverlayMaterial)
 				LastOverlayMaterial = Result;
 			}
 		}
-
-		UE::MediaPlate::Private::ApplyTranslucencyScreenPercentageCVar(1);
 	}
 }
 
diff --git a/Engine/Plugins/Media/MediaPlate/Source/MediaPlate/Private/MediaPlateComponent.cpp b/Engine/Plugins/Media/MediaPlate/Source/MediaPlate/Private/MediaPlateComponent.cpp
--- a/Engine/Plugins/Media/MediaPlate/Source/MediaPlate/Private/MediaPlateComponent.cpp
+++ b/Engine/Plugins/Media/MediaPlate/Source/MediaPlate/Private/MediaPlateComponent.cpp
@@ -703,28 +703,11 @@ void UMediaPlateComponent::RegisterWithMediaTextureTracker()
 	EMediaTextureTargetViewResolution TargetViewResolutionMask = EMediaTextureTargetViewResolution::RenderResolution;
 	if (AMediaPlate* MediaPlate = GetOwner<AMediaPlate>())
 	{
-		// If holdout is enabled, we assume the holdout compositing is active and only use the display resolution.
-		if (MediaPlate->StaticMeshComponent && MediaPlate->StaticMeshComponent->bHoldout)
+		// If holdout composite is enabled, we only use the display resolution.
+		if (MediaPlate->IsHoldoutCompositeEnabled())
 		{
 			TargetViewResolutionMask = EMediaTextureTargetViewResolution::DisplayResolution;
 		}
-		else
-		{
-			// First we check if overlay rendering is done at display resolution
-			static const auto CVarTranslucencySPBasis = IConsoleManager::Get().FindConsoleVariable(TEXT("r.Translucency.ScreenPercentage.Basis"));
-			const bool bDisplayResolutionOverlay = CVarTranslucencySPBasis && CVarTranslucencySPBasis->GetInt() == 1;
-
-			if (bDisplayResolutionOverlay && MediaPlate->GetCurrentOverlayMaterial() != nullptr)
-			{
-				EnumAddFlags(TargetViewResolutionMask, EMediaTextureTargetViewResolution::DisplayResolution);
-
-				// If we only have an overlay material, we can safely remove the render resolution target
-				if (MediaPlate->GetCurrentMaterial() == nullptr)
-				{
-					EnumRemoveFlags(TargetViewResolutionMask, EMediaTextureTargetViewResolution::RenderResolution);
-				}
-			}
-		}
 	}
 	MediaTextureTrackerObject->TargetViewResolutionMask = TargetViewResolutionMask;
 
diff --git a/Engine/Source/Runtime/Core/Public/UObject/FortniteMainBranchObjectVersions.inl b/Engine/Source/Runtime/Core/Public/UObject/FortniteMainBranchObjectVersions.inl
--- a/Engine/Source/Runtime/Core/Public/UObject/FortniteMainBranchObjectVersions.inl
+++ b/Engine/Source/Runtime/Core/Public/UObject/FortniteMainBranchObjectVersions.inl
@@ -634,5 +634,8 @@ FORTNITE_MAIN_VERSION(ThreadContextDataInChaosVisualDebuggerDebugDrawData, 206)
 // Changed default grid mode in surface sampler to a version that's more intuitive and less error-prone
 FORTNITE_MAIN_VERSION(PCGChangedSurfaceSamplerDefaultGridCreationMode, 207)
 
+// Media plate overlay composite technique replacement with holdout composite
+FORTNITE_MAIN_VERSION(MediaPlateOverlayTechniqueRemoval, 208)
+
 // -----<new versions can be added above this line>-------------------------------------------------
 #undef FORTNITE_MAIN_VERSION
\ No newline at end of file
