PR #13399 created by serhanio: Add BlueprintCallable specifier to GetMediaPlaylist

#rnx

diff --git a/Engine/Plugins/Media/MediaPlate/Source/MediaPlate/Public/MediaPlateComponent.h b/Engine/Plugins/Media/MediaPlate/Source/MediaPlate/Public/MediaPlateComponent.h
--- a/Engine/Plugins/Media/MediaPlate/Source/MediaPlate/Public/MediaPlateComponent.h
+++ b/Engine/Plugins/Media/MediaPlate/Source/MediaPlate/Public/MediaPlateComponent.h
@@ -210,6 +210,7 @@ public:
 	/**
 	 * Get the currently active Media Playlist
 	 */
+	UFUNCTION(BlueprintCallable, Category = "Media|MediaPlateComponent")
 	UE_API UMediaPlaylist* GetMediaPlaylist() const;
 
 	/**
