[MediaPlate] Fixing crash with spawnable media plate.

diff --git a/Engine/Plugins/Media/MediaPlate/Source/MediaPlate/Private/MediaPlateComponent.cpp b/Engine/Plugins/Media/MediaPlate/Source/MediaPlate/Private/MediaPlateComponent.cpp
--- a/Engine/Plugins/Media/MediaPlate/Source/MediaPlate/Private/MediaPlateComponent.cpp
+++ b/Engine/Plugins/Media/MediaPlate/Source/MediaPlate/Private/MediaPlateComponent.cpp
@@ -854,6 +854,9 @@ void UMediaPlateComponent::SetLetterboxAspectRatio(float AspectRatio)
 
 void UMediaPlateComponent::SetNumberOfTextures(int32 NumTextures)
 {
+	// Make sure we don't have null elements.
+	MediaTextures.Remove(nullptr);
+
 	if (MediaTextures.Num() != NumTextures)
 	{
 		if (IsRegistered())
