[MediaPlate] Adding GetSelectedMediaSource() accessor function.

#PR 13447

diff --git a/Engine/Plugins/Media/MediaPlate/Source/MediaPlate/Private/MediaPlateComponent.cpp b/Engine/Plugins/Media/MediaPlate/Source/MediaPlate/Private/MediaPlateComponent.cpp
--- a/Engine/Plugins/Media/MediaPlate/Source/MediaPlate/Private/MediaPlateComponent.cpp
+++ b/Engine/Plugins/Media/MediaPlate/Source/MediaPlate/Private/MediaPlateComponent.cpp
@@ -606,6 +606,24 @@ UMediaPlaylist* UMediaPlateComponent::GetMediaPlaylist() const
 	return ActivePlaylist;
 }
 
+UMediaSource* UMediaPlateComponent::GetSelectedMediaSource() const
+{
+	switch (MediaPlateResource.GetResourceType())
+	{
+	case EMediaPlateResourceType::Playlist:
+		return nullptr;
+
+	case EMediaPlateResourceType::External:
+		return ExternalMediaSource;
+
+	case EMediaPlateResourceType::Asset:
+		return MediaPlateResource.GetMediaAsset();
+
+	default:
+		return nullptr;
+	}
+}
+
 void UMediaPlateComponent::SelectExternalMedia(const FString& InFilePath)
 {
 	MediaPlateResource.LoadExternalMedia(InFilePath);
diff --git a/Engine/Plugins/Media/MediaPlate/Source/MediaPlate/Public/MediaPlateComponent.h b/Engine/Plugins/Media/MediaPlate/Source/MediaPlate/Public/MediaPlateComponent.h
--- a/Engine/Plugins/Media/MediaPlate/Source/MediaPlate/Public/MediaPlateComponent.h
+++ b/Engine/Plugins/Media/MediaPlate/Source/MediaPlate/Public/MediaPlateComponent.h
@@ -213,6 +213,12 @@ public:
 	UFUNCTION(BlueprintCallable, Category = "Media|MediaPlateComponent")
 	UE_API UMediaPlaylist* GetMediaPlaylist() const;
 
+	/**
+	 * Get the currently active Media Source.
+	 */
+	UFUNCTION(BlueprintCallable, Category = "Media|MediaPlateComponent")
+	UE_API UMediaSource* GetSelectedMediaSource() const;
+
 	/**
 	 * @brief Select the external media file (non-UE asset) to be opened.
 	 * @param InFilePath non-UFS path (absolute or relative to executable) to the media file.
