[Media Plate] Resetting the asset on a playing media plate will close it.

The main issue was in UMediaPlateComponent::RefreshActivePlaylist() where it didn't reset the playlist if there was no media source.

diff --git a/Engine/Plugins/Media/MediaPlate/Source/MediaPlate/Private/MediaPlateComponent.cpp b/Engine/Plugins/Media/MediaPlate/Source/MediaPlate/Private/MediaPlateComponent.cpp
--- a/Engine/Plugins/Media/MediaPlate/Source/MediaPlate/Private/MediaPlateComponent.cpp
+++ b/Engine/Plugins/Media/MediaPlate/Source/MediaPlate/Private/MediaPlateComponent.cpp
@@ -451,17 +451,23 @@ void UMediaPlateComponent::Open()
 
 	if (IsVisible())
 	{
-		bool bIsPlaying = false;
 		if (MediaPlayer != nullptr)
 		{
-			UMediaSource* MediaSource = UE::MediaPlateComponent::GetMediaSource(GetMediaPlaylist(), PlaylistIndex);
-			bIsPlaying = PlayMediaSource(MediaSource, bPlayOnOpen);
+			if (UMediaSource* MediaSource = UE::MediaPlateComponent::GetMediaSource(GetMediaPlaylist(), PlaylistIndex))
+			{
+				if (!PlayMediaSource(MediaSource, bPlayOnOpen))
+				{
+					UE_LOG(LogMediaPlate, Warning, TEXT("Open Failed: failed to open media source \"%s\"."), *MediaSource->GetUrl());	
+				}
+			}
+			else
+			{
+				UE_LOG(LogMediaPlate, Warning, TEXT("Open Failed: no media source specified."));
+			}
 		}
-
-		// Did anything play?
-		if (bIsPlaying == false)
+		else
 		{
-			UE_LOG(LogMediaPlate, Warning, TEXT("Could not play anything."));
+			UE_LOG(LogMediaPlate, Warning, TEXT("Open Failed: missing media player."));
 		}
 	}
 	else
@@ -1633,9 +1639,12 @@ void UMediaPlateComponent::RefreshActivePlaylist()
 		break;
 	}
 
+	// Note: Must create a new playlist even if no media source is specified in
+	// order to reset the previous media source.
+	ActivePlaylist = NewObject<UMediaPlaylist>(this, NAME_None, RF_Transactional);
+	
 	if (SelectedMediaSource)
 	{
-		ActivePlaylist = NewObject<UMediaPlaylist>(this, NAME_None, RF_Transactional);
 		ActivePlaylist->Add(SelectedMediaSource);
 	}
 }
