[MediaPlate] Media Plate editor: Improving the aspect ratio picker UX.

diff --git a/Engine/Plugins/Media/MediaPlate/Source/MediaPlateEditor/Private/MediaPlateCustomization.cpp b/Engine/Plugins/Media/MediaPlate/Source/MediaPlateEditor/Private/MediaPlateCustomization.cpp
--- a/Engine/Plugins/Media/MediaPlate/Source/MediaPlateEditor/Private/MediaPlateCustomization.cpp
+++ b/Engine/Plugins/Media/MediaPlate/Source/MediaPlateEditor/Private/MediaPlateCustomization.cpp
@@ -31,10 +31,132 @@
 
 #define LOCTEXT_NAMESPACE "FMediaPlateCustomization"
 
+namespace UE::MediaPlateCustomization::Private
+{
+	/** Special aspect ratio value for the "custom" preset entry. */
+	constexpr float CustomAspectRatioPresetValue = -1.0f;
+
+	/** Special aspect ratio value for the "disable" preset entry. */
+	constexpr float DisableAspectRatioPresetValue = 0.0f;
+	
+	struct FAspectRatioPreset
+	{
+		FText DisplayName;
+		FText Description;
+		float AspectRatio = 0;
+	};
+	
+	/**
+	 * Hardcoding the aspect ratio presets for Media Plate.
+	 * We ensure that all entries are unique to allow reversible mapping from the AR value to the combo box. 
+	 */
+	static const TArray<FAspectRatioPreset> AspectRatioPresets =
+	{
+		// Displays and video ARs
+		{
+			LOCTEXT("ARPresetLabel_16:9", "Widescreen (16:9)"),
+			LOCTEXT("ARPresetDesc_16:9", "The standard for high-definition (HD) television, 4K displays, computer monitors, and modern widescreen video."),
+			16.0f/9.0f
+		},
+		{
+			LOCTEXT("ARPresetLabel_16:10", "Widescreen+ (16:10)"),
+			LOCTEXT("ARPresetDesc_16:10", "Used for many widescreen computer displays and some tablets, offering a slightly taller image than 16:9."),
+			16.0f/10.0f
+		},
+		{
+			LOCTEXT("ARPresetLabel_4:3", "Fullscreen (4:3)"),
+			LOCTEXT("ARPresetDesc_4:3", "A traditional, older aspect ratio for standard definition TVs (like NTSC and PAL) and some non-widescreen computer monitors."), 
+			4.0f/3.0f
+		},
+		{
+			LOCTEXT("ARPresetLabel_21:9", "Ultra-Widescreen (21:9)"),
+			LOCTEXT("ARPresetDesc_21:9", "An ultra-widescreen format used for some computer displays and a cinematic feel."),
+			21.0f/9.0f
+		},
+		{
+			LOCTEXT("ARPresetLabel_9:16", "Portrait (9:16)"),
+			LOCTEXT("ARPresetDesc_9:16", "A vertical aspect ratio (a 90-degree flip of 16:9) popular for video stories on social media platforms like Instagram and Snapchat."),
+			9.0f/16.0f
+		},
+		// Photography and Cinematography ARs
+		{
+			LOCTEXT("ARPresetLabel_3:2", "35mm film (3:2)"),
+			LOCTEXT("ARPresetDesc_3:2", "A traditional aspect ratio that originated with 35mm film, still common in still photography and for print sizes."),
+			3.0f/2.0f
+		},
+		{
+			LOCTEXT("ARPresetLabel_1:1", "Square (1:1)"),
+			LOCTEXT("ARPresetDesc_1:1", "A perfect square format, often used for social media images and some large format photography."),
+			1.0f
+		},
+		{
+			LOCTEXT("ARPresetLabel_5:4", "\"Wide\" (5:4)"),
+			LOCTEXT("ARPresetDesc_5:4", "Used in some photography and art prints, especially with large and medium format cameras."),
+			5.0f/4.0f
+		},
+		{
+			LOCTEXT("ARPresetLabel_1.85:1", "Flat (1.85:1)"),
+			LOCTEXT("ARPresetDesc_1.85:1", "A common aspect ratio for cinema, often referred to as \"flat\"."),
+			1.85f/1.0f
+		},
+		{
+			LOCTEXT("ARPresetLabel_2.39:1", "Scope (2.39:1)" ),
+			LOCTEXT("ARPresetDesc_2.39:1", "Another cinematic aspect ratio, also known as the \"scope\" format."),
+			2.39f/1.0f
+		},
+	};
+
+	/** Special preset entry for custom entries. */
+	static const FAspectRatioPreset CustomAspectRatioPreset =
+	{
+		LOCTEXT("ARPresetLabel_Custom", "Custom" ),
+		LOCTEXT("ARPresetDesc_Custom", "Manually entered aspect ratio not part of presets."),
+		CustomAspectRatioPresetValue
+	};
+
+	/** Special preset entry to disable the feature (used for letter box ar). */
+	static const FAspectRatioPreset DisabledAspectRatioPreset =
+	{
+		LOCTEXT("Disable", "Disable" ),	// Note: Keeping same name to not change localisation.
+		LOCTEXT("Disable_Desc", "Disable the feature."),
+		DisableAspectRatioPresetValue
+	};
+	
+	/**
+	 * Finds the corresponding preset entry for the given aspect ratio value.
+	 * @param InAspectRatio Aspect ratio to search for.
+	 * @param bInIncludeDisable Indicate if "disable" entry should be returned.
+	 * @return Found preset entry or "custom" entry if not found.
+	 */
+	const FAspectRatioPreset& FindAspectRatioPreset(float InAspectRatio, bool bInIncludeDisable)
+	{
+		for (const FAspectRatioPreset& Preset : AspectRatioPresets)
+		{
+			if (Preset.AspectRatio == InAspectRatio)
+			{
+				return Preset;
+			}
+		}
+
+		// Check if this is the "disable" entry.
+		if (bInIncludeDisable && InAspectRatio == DisableAspectRatioPresetValue)
+		{
+			return DisabledAspectRatioPreset;
+		}
+
+		// If not found, return the custom preset entry.
+		return CustomAspectRatioPreset;
+	}
+
+	const FText& FindAspectRatioDisplayName(float InAspectRatio, bool bInIncludeDisable)
+	{
+		return FindAspectRatioPreset(InAspectRatio, bInIncludeDisable).DisplayName;
+	}
+}
+
 /* IDetailCustomization interface
  *****************************************************************************/
 
-
 FMediaPlateCustomization::FMediaPlateCustomization()
 {
 }
@@ -467,7 +589,7 @@ void FMediaPlateCustomization::AddMeshCustomization(IDetailCategoryBuilder& InPa
 						[
 							SNew(STextBlock)
 								.ToolTipText(LOCTEXT("Presets_ToolTip", "Select one of the presets for the aspect ratio."))
-								.Text(LOCTEXT("Presets", "Presets"))
+								.Text(this, &FMediaPlateCustomization::OnGetAspectRatioText)
 						]
 				]
 
@@ -509,7 +631,7 @@ void FMediaPlateCustomization::AddMeshCustomization(IDetailCategoryBuilder& InPa
 						[
 							SNew(STextBlock)
 								.ToolTipText(LOCTEXT("Presets_ToolTip", "Select one of the presets for the aspect ratio."))
-								.Text(LOCTEXT("Presets", "Presets"))
+								.Text(this, &FMediaPlateCustomization::OnGetLetterboxAspectRatioText)
 						]
 			]
 
@@ -708,22 +830,35 @@ TSharedRef<SWidget> FMediaPlateCustomization::OnGetLetterboxAspectRatios()
 	FMenuBuilder MenuBuilder(true, NULL);
 	AddAspectRatiosToMenuBuilder(MenuBuilder, &FMediaPlateCustomization::SetLetterboxAspectRatio);
 
-	FUIAction Action = FUIAction(FExecuteAction::CreateSP(this, &FMediaPlateCustomization::SetLetterboxAspectRatio, 0.0f));
-	MenuBuilder.AddMenuEntry(LOCTEXT("Disable", "Disable"), FText(), FSlateIcon(), Action);
+	// Add the extra "disable" entry.
+	using namespace UE::MediaPlateCustomization::Private; 
+	FUIAction Action = FUIAction(FExecuteAction::CreateSP(this, &FMediaPlateCustomization::SetLetterboxAspectRatio, DisabledAspectRatioPreset.AspectRatio));
+	MenuBuilder.AddMenuEntry(DisabledAspectRatioPreset.DisplayName, DisabledAspectRatioPreset.Description, FSlateIcon(), Action);
 
 	return MenuBuilder.MakeWidget();
 }
 
+FText FMediaPlateCustomization::OnGetAspectRatioText() const
+{
+	using namespace UE::MediaPlateCustomization::Private;
+	return FindAspectRatioDisplayName(GetAspectRatio(), /*bInIncludeDisable*/false);
+}
+
+FText FMediaPlateCustomization::OnGetLetterboxAspectRatioText() const
+{
+	using namespace UE::MediaPlateCustomization::Private;
+	return FindAspectRatioDisplayName(GetLetterboxAspectRatio(), /*bInIncludeDisable*/true);
+}
+
 void FMediaPlateCustomization::AddAspectRatiosToMenuBuilder(FMenuBuilder& MenuBuilder,
 	void (FMediaPlateCustomization::* Func)(float))
 {
-	TArray<FNamedFilmbackPreset> const& Presets = UCineCameraSettings::GetFilmbackPresets();
-
-	for (const FNamedFilmbackPreset& Preset : Presets)
+	using namespace UE::MediaPlateCustomization::Private;
+	for (const FAspectRatioPreset& Preset : AspectRatioPresets)
 	{
 		FUIAction Action = FUIAction(FExecuteAction::CreateSP(this,
-			Func, Preset.FilmbackSettings.SensorAspectRatio));
-		MenuBuilder.AddMenuEntry(FText::FromString(Preset.Name), FText(), FSlateIcon(), Action);
+			Func, Preset.AspectRatio));
+		MenuBuilder.AddMenuEntry(Preset.DisplayName, Preset.Description, FSlateIcon(), Action);
 	}
 }
 
diff --git a/Engine/Plugins/Media/MediaPlate/Source/MediaPlateEditor/Private/MediaPlateCustomization.h b/Engine/Plugins/Media/MediaPlate/Source/MediaPlateEditor/Private/MediaPlateCustomization.h
--- a/Engine/Plugins/Media/MediaPlate/Source/MediaPlateEditor/Private/MediaPlateCustomization.h
+++ b/Engine/Plugins/Media/MediaPlate/Source/MediaPlateEditor/Private/MediaPlateCustomization.h
@@ -51,7 +51,7 @@ private:
 	TArray<TWeakObjectPtr<UMediaPlateComponent>> MediaPlatesList;
 
 	/** Whether we have a plane, sphere, etc. */
-	EMediaTextureVisibleMipsTiles MeshMode;
+	EMediaTextureVisibleMipsTiles MeshMode = EMediaTextureVisibleMipsTiles::None;
 
 	/** Handles mesh stuff. */
 	FMediaPlateCustomizationMesh MeshCustomization;
@@ -99,6 +99,16 @@ private:
 	 */
 	TSharedRef<SWidget> OnGetLetterboxAspectRatios();
 
+	/**
+	 * Returns the text to display for a given aspect ratio.
+	 */
+	FText OnGetAspectRatioText() const;
+	
+	/**
+	 * Returns the text to display for a given letter box aspect ratio.
+	 */
+	FText OnGetLetterboxAspectRatioText() const;
+	
 	/**
 	 * Adds menu options for all aspect ratio presets.
 	 */
