[Media Plate] Fixing Letter box

It looks like the referenced "/MediaPlate/M_MediaPlateLetterbox" material was never added. The code change referencing that asset was done in CL 21321261.
An easy fix is to use the default opaque material, which is black by default.

diff --git a/Engine/Plugins/Media/MediaPlate/Source/MediaPlate/Private/MediaPlateComponent.cpp b/Engine/Plugins/Media/MediaPlate/Source/MediaPlate/Private/MediaPlateComponent.cpp
--- a/Engine/Plugins/Media/MediaPlate/Source/MediaPlate/Private/MediaPlateComponent.cpp
+++ b/Engine/Plugins/Media/MediaPlate/Source/MediaPlate/Private/MediaPlateComponent.cpp
@@ -1281,7 +1281,7 @@ void UMediaPlateComponent::AddLetterboxes()
 		if (Owner != nullptr)
 		{
 			UStaticMesh* Mesh = LoadObject<UStaticMesh>(NULL, TEXT("/MediaPlate/SM_MediaPlateScreen"), NULL, LOAD_None, NULL);
-			UMaterial* Material = LoadObject<UMaterial>(NULL, TEXT("/MediaPlate/M_MediaPlateLetterbox"), NULL, LOAD_None, NULL);
+			UMaterial* Material = LoadObject<UMaterial>(NULL, TEXT("/MediaPlate/M_MediaPlate_Opaque"), NULL, LOAD_None, NULL);
 			if ((Mesh != nullptr) && (Material != nullptr))
 			{
 				for (int32 Index = 0; Index < 2; ++Index)
