[MediaPlate] MediaPlate Track Editor: implement "apply to all" for the questions to user on multiple selection.

diff --git a/Engine/Plugins/Media/MediaPlate/Source/MediaPlateEditor/Private/Sequencer/MediaPlateTrackEditor.cpp b/Engine/Plugins/Media/MediaPlate/Source/MediaPlateEditor/Private/Sequencer/MediaPlateTrackEditor.cpp
--- a/Engine/Plugins/Media/MediaPlate/Source/MediaPlateEditor/Private/Sequencer/MediaPlateTrackEditor.cpp
+++ b/Engine/Plugins/Media/MediaPlate/Source/MediaPlateEditor/Private/Sequencer/MediaPlateTrackEditor.cpp
@@ -10,6 +10,7 @@
 #include "MediaPlate.h"
 #include "MediaPlateComponent.h"
 #include "MediaPlaylist.h"
+#include "Misc/ITransaction.h"
 #include "Misc/MessageDialog.h"
 #include "Modules/ModuleManager.h"
 #include "MovieSceneObjectBindingID.h"
@@ -185,19 +186,31 @@ void FMediaPlateTrackEditor::AddTrackForComponent(UMediaPlateComponent* Componen
 	// Make sure the media plate's player is not playing
 	Component->Close();
 
+	// Helper function to implement the "apply to all" logic and response tracking over multiple operations.
+	auto AskForDisable = [](const FText& InMessage, FDialogResponseInfo& InResponseInfo)
+	{
+		// We are using the current transaction id to establish a context for our response info.
+		const FGuid CurrentContextId = GUndo ? GUndo->GetContext().TransactionId : FGuid();
+
+		// See if we need to ask again or can reuse previous answer.	
+		if (!InResponseInfo.ContextId.IsValid() || InResponseInfo.ContextId != CurrentContextId || !InResponseInfo.bAll)
+		{
+			const EAppReturnType::Type Reply = FMessageDialog::Open(EAppMsgType::YesNoYesAllNoAll, InMessage);
+			InResponseInfo.ContextId = CurrentContextId;
+			InResponseInfo.bAll = Reply == EAppReturnType::YesAll || Reply == EAppReturnType::NoAll;
+			InResponseInfo.bApply = Reply == EAppReturnType::Yes || Reply == EAppReturnType::YesAll;
+		}
+		return InResponseInfo.bApply;
+	};
+
 	// Does the media plate have autoplay?
 	if (Component->bAutoPlay)
 	{
 		// Ask the user if it can be disabled.
-		EAppReturnType::Type YesNoCancelReply = FMessageDialog::Open(EAppMsgType::YesNo,
-			LOCTEXT("Prompt_DisableAutoplay", "Would you like to disable autoplay on the MediaPlate actor (recommended)?"));
-
-		switch (YesNoCancelReply)
+		if (AskForDisable(LOCTEXT("Prompt_DisableAutoplay", "Would you like to disable autoplay on the MediaPlate actor (recommended)?"), LastAutoPlayDisableResponse))
 		{
-			case EAppReturnType::Yes:
-				Component->Modify();
-				Component->bAutoPlay = false;
-				break;
+			Component->Modify();
+			Component->bAutoPlay = false;
 		}
 	}
 
@@ -205,10 +218,7 @@ void FMediaPlateTrackEditor::AddTrackForComponent(UMediaPlateComponent* Componen
 	if (Component->bEnableAudio)
 	{
 		// Ask the user if it can be disabled.
-		const EAppReturnType::Type Reply = FMessageDialog::Open(EAppMsgType::YesNo,
-			LOCTEXT("Prompt_DisableAudio", "Would you like to disable audio on the MediaPlate actor (recommended)?"));
-
-		if (Reply == EAppReturnType::Yes)
+		if (AskForDisable(LOCTEXT("Prompt_DisableAudio", "Would you like to disable audio on the MediaPlate actor (recommended)?"), LastAudioDisableResponse))
 		{
 			Component->Modify();
 			Component->SetEnableAudio(false);
diff --git a/Engine/Plugins/Media/MediaPlate/Source/MediaPlateEditor/Private/Sequencer/MediaPlateTrackEditor.h b/Engine/Plugins/Media/MediaPlate/Source/MediaPlateEditor/Private/Sequencer/MediaPlateTrackEditor.h
--- a/Engine/Plugins/Media/MediaPlate/Source/MediaPlateEditor/Private/Sequencer/MediaPlateTrackEditor.h
+++ b/Engine/Plugins/Media/MediaPlate/Source/MediaPlateEditor/Private/Sequencer/MediaPlateTrackEditor.h
@@ -94,4 +94,20 @@ private:
 
 	/** Handle to our delegate. */
 	FDelegateHandle OnActorAddedToSequencerHandle;
+
+	/**
+	 * Information on the last dialog response.
+	 */
+	struct FDialogResponseInfo
+	{
+		/** Keep track of the context id for the response to all during multiple operations. */
+		FGuid ContextId;
+		/** True if the response applies to all in multiple operations under the same context id. */
+		bool bAll = false;
+		/** True if the operation should be done. */
+		bool bApply = false;
+	};
+
+	FDialogResponseInfo LastAutoPlayDisableResponse;
+	FDialogResponseInfo LastAudioDisableResponse;
 };
