[Media Plate] Media Plate Track Editor: handle AddTrack failure (crash prevention).

diff --git a/Engine/Plugins/Media/MediaPlate/Source/MediaPlateEditor/Private/Sequencer/MediaPlateTrackEditor.cpp b/Engine/Plugins/Media/MediaPlate/Source/MediaPlateEditor/Private/Sequencer/MediaPlateTrackEditor.cpp
--- a/Engine/Plugins/Media/MediaPlate/Source/MediaPlateEditor/Private/Sequencer/MediaPlateTrackEditor.cpp
+++ b/Engine/Plugins/Media/MediaPlate/Source/MediaPlateEditor/Private/Sequencer/MediaPlateTrackEditor.cpp
@@ -112,13 +112,15 @@ void FMediaPlateTrackEditor::HandleAddMediaTrackToObjectBindingMenuEntryExecute(
 	{
 		if (InObjectBindingID.IsValid())
 		{
-			// Add media track.
-			UMovieSceneMediaTrack* NewObjectTrack = FocusedMovieScene->AddTrack<UMovieSceneMediaTrack>(InObjectBindingID);
-			NewObjectTrack->SetDisplayName(LOCTEXT("MediaTrackName", "Media"));
-
-			if (GetSequencer().IsValid())
+			// Add media track.			
+			if (UMovieSceneMediaTrack* NewObjectTrack = FocusedMovieScene->AddTrack<UMovieSceneMediaTrack>(InObjectBindingID))
 			{
-				GetSequencer()->OnAddTrack(NewObjectTrack, InObjectBindingID);
+				NewObjectTrack->SetDisplayName(LOCTEXT("MediaTrackName", "Media"));
+
+				if (GetSequencer().IsValid())
+				{
+					GetSequencer()->OnAddTrack(NewObjectTrack, InObjectBindingID);
+				}
 			}
 		}
 	}
