[MediaPlate] Reset the media texture on asset change.

diff --git a/Engine/Plugins/Media/MediaPlate/Source/MediaPlate/Private/MediaPlateComponent.cpp b/Engine/Plugins/Media/MediaPlate/Source/MediaPlate/Private/MediaPlateComponent.cpp
--- a/Engine/Plugins/Media/MediaPlate/Source/MediaPlate/Private/MediaPlateComponent.cpp
+++ b/Engine/Plugins/Media/MediaPlate/Source/MediaPlate/Private/MediaPlateComponent.cpp
@@ -69,6 +69,18 @@ namespace UE::MediaPlateComponent
 #endif
 	}
 
+	// Find the media texture associated with the given media player and update it's resource to reset it.
+	void ResetMediaTextureForPlayer(const UMediaPlayer* InMediaPlayer, const TArray<TObjectPtr<UMediaTexture>>& InMediaTextures)
+	{
+		for (TObjectPtr<UMediaTexture> MediaTexture : InMediaTextures)
+		{
+			if (MediaTexture && MediaTexture->GetMediaPlayer() == InMediaPlayer)
+			{
+				MediaTexture->UpdateResource();
+			}
+		}
+	}
+
 	UMediaSource* GetMediaSource(UMediaPlaylist* InPlaylist, int32 InPlaylistIndex)
 	{
 		return InPlaylist ? InPlaylist->Get(InPlaylistIndex) : nullptr;
@@ -1578,22 +1590,24 @@ void UMediaPlateComponent::RefreshMediaPlateResource()
 	RefreshExternalMediaSource();
 	RefreshActivePlaylist();
 
+	const UMediaSource* SelectedMediaSource = UE::MediaPlateComponent::GetMediaSource(GetMediaPlaylist(), PlaylistIndex);
+
 	// Remote Control may call PostEditChangeProperty more than once from various paths
 	// with some delays causing Media Plate to close when it shouldn't. We mitigate that by
 	// skipping the closing/restarting logic if the media source didn't change.
-	if (MediaPlayer)
+	if (MediaPlayer && SelectedMediaSource)
 	{
-		const UMediaSource* MediaSource = UE::MediaPlateComponent::GetMediaSource(GetMediaPlaylist(), PlaylistIndex);
 		const UMediaSource* PlayerMediaSource = UE::MediaPlateComponent::GetMediaSource(MediaPlayer->GetPlaylist(), MediaPlayer->GetPlaylistIndex());
-		if (MediaSource == PlayerMediaSource)
+		if (SelectedMediaSource == PlayerMediaSource)
 		{
 			return;
 		}
-	}
 
-	if (RestartPlayer())
-	{
-		return;
+		// Only try to restart the player if there is a new selected media source.
+		if (RestartPlayer())
+		{
+			return;
+		}
 	}
 
 	constexpr EMediaPlateEventState State = EMediaPlateEventState::Close;
@@ -1603,6 +1617,12 @@ void UMediaPlateComponent::RefreshMediaPlateResource()
 	{
 		MediaAssets->BroadcastOnMediaStateChangedEvent({GetOwner()->GetPathName()}, static_cast<uint8>(State));
 	}
+
+	// If the media source has changed, reset the associated media texture to clear the stale
+	// video sample from the previous asset.
+	// Note: if exposed to media track, the texture is not cleared if the track is evaluating, i.e.
+	// the track has its own player associated to the texture and will not be affected by this.
+	UE::MediaPlateComponent::ResetMediaTextureForPlayer(MediaPlayer, MediaTextures);
 }
 
 void UMediaPlateComponent::RefreshExternalMediaSource()
