[Media Plate] Fix Aspect Ratio and Mesh Mode transactions to fix Multi-User operations.

diff --git a/Engine/Plugins/Media/MediaPlate/Source/MediaPlateEditor/Private/MediaPlateCustomization.cpp b/Engine/Plugins/Media/MediaPlate/Source/MediaPlateEditor/Private/MediaPlateCustomization.cpp
--- a/Engine/Plugins/Media/MediaPlate/Source/MediaPlateEditor/Private/MediaPlateCustomization.cpp
+++ b/Engine/Plugins/Media/MediaPlate/Source/MediaPlateEditor/Private/MediaPlateCustomization.cpp
@@ -747,6 +747,8 @@ void FMediaPlateCustomization::SetMeshMode(EMediaTextureVisibleMipsTiles InMode)
 			UMediaPlateComponent* MediaPlate = MediaPlatePtr.Get();
 			if (MediaPlate != nullptr)
 			{
+				MediaPlate->Modify();
+
 				// Update the setting in the media plate.
 				MediaPlate->SetVisibleMipsTilesCalculations(MeshMode);
 
@@ -804,6 +806,7 @@ ECheckBoxState FMediaPlateCustomization::IsAspectRatioAuto() const
 
 void FMediaPlateCustomization::SetIsAspectRatioAuto(ECheckBoxState State)
 {
+	const FScopedTransaction Transaction(LOCTEXT("SetIsAspectRatioAuto", "Media Plate Set Is Aspect Ratio Auto"));
 	bool bEnable = (State == ECheckBoxState::Checked);
 
 	// Loop through all our objects.
@@ -812,6 +815,7 @@ void FMediaPlateCustomization::SetIsAspectRatioAuto(ECheckBoxState State)
 		UMediaPlateComponent* MediaPlate = MediaPlatePtr.Get();
 		if (MediaPlate != nullptr)
 		{
+			MediaPlate->Modify();
 			MediaPlate->SetIsAspectRatioAuto(bEnable);
 		}
 	}
