[MediaPlateEditor] Fix the media plate editor's playback buttons.

Unifying the implementation with the media plate customization buttons.

diff --git a/Engine/Plugins/Media/MediaPlate/Source/MediaPlateEditor/Private/MediaPlateCustomization.cpp b/Engine/Plugins/Media/MediaPlate/Source/MediaPlateEditor/Private/MediaPlateCustomization.cpp
--- a/Engine/Plugins/Media/MediaPlate/Source/MediaPlateEditor/Private/MediaPlateCustomization.cpp
+++ b/Engine/Plugins/Media/MediaPlate/Source/MediaPlateEditor/Private/MediaPlateCustomization.cpp
@@ -1055,7 +1055,7 @@ bool FMediaPlateCustomization::IsButtonEventAllowedForAnyPlate(EMediaPlateEventS
 		if (UMediaPlateComponent* MediaPlate = MediaPlateWeak.Get())
 		{
 			if (MediaPlate->IsEventStateChangeAllowed(InState)
-				&& IsButtonEventAllowedForPlayer(InState, MediaPlate->GetMediaPlayer()))
+				&& IsMediaPlateEventAllowedForPlayer(InState, MediaPlate->GetMediaPlayer()))
 			{
 				return true;
 			}
@@ -1064,7 +1064,7 @@ bool FMediaPlateCustomization::IsButtonEventAllowedForAnyPlate(EMediaPlateEventS
 	return false;
 }
 
-bool FMediaPlateCustomization::IsButtonEventAllowedForPlayer(EMediaPlateEventState InState, UMediaPlayer* InMediaPlayer) const
+bool FMediaPlateCustomization::IsMediaPlateEventAllowedForPlayer(EMediaPlateEventState InState, UMediaPlayer* InMediaPlayer)
 {
 	// Note: centralize the state switch conditions here to make it easier to maintain.
 	if (!InMediaPlayer)
@@ -1106,16 +1106,21 @@ bool FMediaPlateCustomization::IsButtonEventAllowedForPlayer(EMediaPlateEventSta
 }
 
 void FMediaPlateCustomization::OnButtonEvent(EMediaPlateEventState InState)
+{
+	HandleMediaPlateEvent(MediaPlatesList, InState);
+}
+
+void FMediaPlateCustomization::HandleMediaPlateEvent(TConstArrayView<TWeakObjectPtr<UMediaPlateComponent>> InMediaPlates, EMediaPlateEventState InState)
 {
 	IMediaAssetsModule* MediaAssets = FModuleManager::LoadModulePtr<IMediaAssetsModule>("MediaAssets");
 	TArray<FString> ActorsPathNames;
-	ActorsPathNames.Reserve(MediaPlatesList.Num());
-	for (const TWeakObjectPtr<UMediaPlateComponent>& MediaPlateWeak : MediaPlatesList)
+	ActorsPathNames.Reserve(InMediaPlates.Num());
+	for (const TWeakObjectPtr<UMediaPlateComponent>& MediaPlateWeak : InMediaPlates)
 	{
 		UMediaPlateComponent* MediaPlate = MediaPlateWeak.Get();
 
 		// Note: because of multi-selection and the possibility of different player states, we need check restrictions per plate again.
-		if (!MediaPlate || !MediaPlate->IsEventStateChangeAllowed(InState) || !IsButtonEventAllowedForPlayer(InState, MediaPlate->GetMediaPlayer()))
+		if (!MediaPlate || !MediaPlate->IsEventStateChangeAllowed(InState) || !IsMediaPlateEventAllowedForPlayer(InState, MediaPlate->GetMediaPlayer()))
 		{
 			continue;
 		}
diff --git a/Engine/Plugins/Media/MediaPlate/Source/MediaPlateEditor/Private/MediaPlateCustomization.h b/Engine/Plugins/Media/MediaPlate/Source/MediaPlateEditor/Private/MediaPlateCustomization.h
--- a/Engine/Plugins/Media/MediaPlate/Source/MediaPlateEditor/Private/MediaPlateCustomization.h
+++ b/Engine/Plugins/Media/MediaPlate/Source/MediaPlateEditor/Private/MediaPlateCustomization.h
@@ -42,6 +42,17 @@ public:
 		return MakeShareable(new FMediaPlateCustomization());
 	}
 
+	/**
+	* Determines if the event state switch is allowed for the given media player.
+	* Those are the UI only conditions to mirror the backend's conditions.
+	*/
+	static bool IsMediaPlateEventAllowedForPlayer(EMediaPlateEventState InState, UMediaPlayer* InMediaPlayer);
+
+	/**
+	* Changes the state of selected media plates and broadcasts the event to the remote endpoints.
+	*/
+	static void HandleMediaPlateEvent(TConstArrayView<TWeakObjectPtr<UMediaPlateComponent>> InMediaPlates, EMediaPlateEventState InState);
+
 private:
 	static void CustomizeCategories(IDetailLayoutBuilder& InDetailBuilder);
 
@@ -186,13 +197,7 @@ private:
 	bool IsButtonEventAllowedForAnyPlate(EMediaPlateEventState InState) const;
 
 	/**
-	 * Determines if the event state switch is allowed for the given media player.
-	 * Those are the UI only conditions to mirror the backend's conditions.
-	 */
-	bool IsButtonEventAllowedForPlayer(EMediaPlateEventState InState, UMediaPlayer* InMediaPlayer) const;
-
-	/**
-	* Changes the state of selected media plates and broadcasts the even to the remote endpoints.
+	* Changes the state of selected media plates and broadcasts the event to the remote endpoints.
 	*/
 	void OnButtonEvent(EMediaPlateEventState InState);
 
diff --git a/Engine/Plugins/Media/MediaPlate/Source/MediaPlateEditor/Private/Toolkits/MediaPlateEditorToolkit.cpp b/Engine/Plugins/Media/MediaPlate/Source/MediaPlateEditor/Private/Toolkits/MediaPlateEditorToolkit.cpp
--- a/Engine/Plugins/Media/MediaPlate/Source/MediaPlateEditor/Private/Toolkits/MediaPlateEditorToolkit.cpp
+++ b/Engine/Plugins/Media/MediaPlate/Source/MediaPlateEditor/Private/Toolkits/MediaPlateEditorToolkit.cpp
@@ -9,6 +9,7 @@
 #include "IMediaAssetsModule.h"
 #include "MediaPlate.h"
 #include "MediaPlateComponent.h"
+#include "MediaPlateCustomization.h"
 #include "MediaPlateEditorModule.h"
 #include "MediaPlayer.h"
 #include "MediaPlaylist.h"
@@ -164,11 +165,7 @@ FString FMediaPlateEditorToolkit::GetDocumentationLink() const
 
 void FMediaPlateEditorToolkit::OnClose()
 {
-	TObjectPtr<UMediaPlayer> MediaPlayer = MediaPlate->GetMediaPlayer();
-	if (MediaPlayer != nullptr)
-	{
-		MediaPlayer->Close();
-	}
+	HandleMediaPlateEvent(EMediaPlateEventState::Close);
 }
 
 void FMediaPlateEditorToolkit::RegisterTabSpawners(const TSharedRef<FTabManager>& InTabManager)
@@ -263,136 +260,59 @@ void FMediaPlateEditorToolkit::PostRedo(bool bSuccess)
 void FMediaPlateEditorToolkit::BindCommands()
 {
 	const FMediaPlateEditorCommands& Commands = FMediaPlateEditorCommands::Get();
-	IMediaAssetsModule* MediaAssets = FModuleManager::LoadModulePtr<IMediaAssetsModule>("MediaAssets");
 
 	ToolkitCommands->MapAction(
 		Commands.CloseMedia,
-		FExecuteAction::CreateLambda([this, MediaAssets] {
-			MediaPlate->GetMediaPlayer()->Close();
-
-			TArray<FString> ActorsPathNames = { MediaPlate->GetOwner()->GetPathName() };
-			MediaAssets->BroadcastOnMediaStateChangedEvent(ActorsPathNames, (uint8)EMediaPlateEventState::Close);
-		}),
-		FCanExecuteAction::CreateLambda([this] {
-			TObjectPtr<UMediaPlayer> MediaPlayer = MediaPlate->GetMediaPlayer();
-			return (MediaPlayer != nullptr) && !MediaPlayer->GetUrl().IsEmpty();
-		})
+		FExecuteAction::CreateSP( this, &FMediaPlateEditorToolkit::HandleMediaPlateEvent, EMediaPlateEventState::Close ),
+		FCanExecuteAction::CreateSP( this, &FMediaPlateEditorToolkit::CanHandleMediaPlateEvent, EMediaPlateEventState::Close )
 	);
 
 	ToolkitCommands->MapAction(
 		Commands.ForwardMedia,
-		FExecuteAction::CreateLambda([this, MediaAssets] {
-			MediaPlate->GetMediaPlayer()->SetRate(GetForwardRate());
-
-			TArray<FString> ActorsPathNames = { MediaPlate->GetOwner()->GetPathName() };
-			MediaAssets->BroadcastOnMediaStateChangedEvent(ActorsPathNames, (uint8)EMediaPlateEventState::Forward);
-		}),
-		FCanExecuteAction::CreateLambda([this]{
-			TObjectPtr<UMediaPlayer> MediaPlayer = MediaPlate->GetMediaPlayer();
-			return (MediaPlayer != nullptr) && MediaPlayer->IsReady() && MediaPlayer->SupportsRate(GetForwardRate(), false);
-		})
+		FExecuteAction::CreateSP( this, &FMediaPlateEditorToolkit::HandleMediaPlateEvent, EMediaPlateEventState::Forward ),
+		FCanExecuteAction::CreateSP( this, &FMediaPlateEditorToolkit::CanHandleMediaPlateEvent, EMediaPlateEventState::Forward )
 	);
 
 	ToolkitCommands->MapAction(
 		Commands.NextMedia,
-		FExecuteAction::CreateLambda([this, MediaAssets] {
-			MediaPlate->Next();
-
-			TArray<FString> ActorsPathNames = { MediaPlate->GetOwner()->GetPathName() };
-			MediaAssets->BroadcastOnMediaStateChangedEvent(ActorsPathNames, (uint8)EMediaPlateEventState::Next);
-		}),
-		FCanExecuteAction::CreateLambda([this]{
-			return (MediaPlate->GetMediaPlaylist() != nullptr) &&
-				(MediaPlate->GetMediaPlaylist()->Num() > 1);
-		})
+		FExecuteAction::CreateSP( this, &FMediaPlateEditorToolkit::HandleMediaPlateEvent, EMediaPlateEventState::Next ),
+		FCanExecuteAction::CreateSP( this, &FMediaPlateEditorToolkit::CanHandleMediaPlateEvent, EMediaPlateEventState::Next )
 	);
 
 	ToolkitCommands->MapAction(
 		Commands.OpenMedia,
-		FExecuteAction::CreateLambda([this, MediaAssets] {
-			MediaPlate->Open();
-
-			// Tell the editor module that this media plate is playing.
-			FMediaPlateEditorModule* EditorModule = FModuleManager::LoadModulePtr<FMediaPlateEditorModule>("MediaPlateEditor");
-			if (EditorModule != nullptr)
-			{
-				EditorModule->MediaPlateStartedPlayback(MediaPlate);
-			}
-
-			TArray<FString> ActorsPathNames = { MediaPlate->GetOwner()->GetPathName() };
-			MediaAssets->BroadcastOnMediaStateChangedEvent(ActorsPathNames, (uint8)EMediaPlateEventState::Open);
-		}),
-		FCanExecuteAction::CreateLambda([this] { return true; })
+		FExecuteAction::CreateSP( this, &FMediaPlateEditorToolkit::HandleMediaPlateEvent, EMediaPlateEventState::Open ),
+		FCanExecuteAction::CreateSP( this, &FMediaPlateEditorToolkit::CanHandleMediaPlateEvent, EMediaPlateEventState::Open )
 	);
 
 	ToolkitCommands->MapAction(
 		Commands.PauseMedia,
-		FExecuteAction::CreateLambda([this, MediaAssets] {
-			MediaPlate->Pause();
-
-			TArray<FString> ActorsPathNames = { MediaPlate->GetOwner()->GetPathName() };
-			MediaAssets->BroadcastOnMediaStateChangedEvent(ActorsPathNames, (uint8)EMediaPlateEventState::Pause);
-		}),
-		FCanExecuteAction::CreateLambda([this]{
-			TObjectPtr<UMediaPlayer> MediaPlayer = MediaPlate->GetMediaPlayer();
-			return (MediaPlayer != nullptr) &&MediaPlayer->CanPause() && !MediaPlayer->IsPaused();
-		})
+		FExecuteAction::CreateSP( this, &FMediaPlateEditorToolkit::HandleMediaPlateEvent, EMediaPlateEventState::Pause ),
+		FCanExecuteAction::CreateSP( this, &FMediaPlateEditorToolkit::CanHandleMediaPlateEvent, EMediaPlateEventState::Pause )
 	);
 
 	ToolkitCommands->MapAction(
 		Commands.PlayMedia,
-		FExecuteAction::CreateLambda([this, MediaAssets] {
-			MediaPlate->GetMediaPlayer()->Play();
-			
-			TArray<FString> ActorsPathNames = { MediaPlate->GetOwner()->GetPathName() };
-			MediaAssets->BroadcastOnMediaStateChangedEvent(ActorsPathNames, (uint8)EMediaPlateEventState::Play);
-		}),
-		FCanExecuteAction::CreateLambda([this]{
-			TObjectPtr<UMediaPlayer> MediaPlayer = MediaPlate->GetMediaPlayer(); 
-			return (MediaPlayer != nullptr) && MediaPlayer->IsReady() && (!MediaPlayer->IsPlaying() || (MediaPlayer->GetRate() != 1.0f));
-		})
+		FExecuteAction::CreateSP( this, &FMediaPlateEditorToolkit::HandleMediaPlateEvent, EMediaPlateEventState::Play ),
+		FCanExecuteAction::CreateSP( this, &FMediaPlateEditorToolkit::CanHandleMediaPlateEvent, EMediaPlateEventState::Play )
 	);
 
 	ToolkitCommands->MapAction(
 		Commands.PreviousMedia,
-		FExecuteAction::CreateLambda([this, MediaAssets] {
-			MediaPlate->Previous();
-
-			TArray<FString> ActorsPathNames = { MediaPlate->GetOwner()->GetPathName() };
-			MediaAssets->BroadcastOnMediaStateChangedEvent(ActorsPathNames, (uint8)EMediaPlateEventState::Previous);
-		}),
-		FCanExecuteAction::CreateLambda([this]{
-			return (MediaPlate->GetMediaPlaylist() != nullptr) &&
-					(MediaPlate->GetMediaPlaylist()->Num() > 1);
-		})
+		FExecuteAction::CreateSP( this, &FMediaPlateEditorToolkit::HandleMediaPlateEvent, EMediaPlateEventState::Previous ),
+		FCanExecuteAction::CreateSP( this, &FMediaPlateEditorToolkit::CanHandleMediaPlateEvent, EMediaPlateEventState::Previous )
 	);
 
 	ToolkitCommands->MapAction(
 		Commands.ReverseMedia,
-		FExecuteAction::CreateLambda([this, MediaAssets] {
-			MediaPlate->GetMediaPlayer()->SetRate(GetReverseRate());
-
-			TArray<FString> ActorsPathNames = { MediaPlate->GetOwner()->GetPathName() };
-			MediaAssets->BroadcastOnMediaStateChangedEvent(ActorsPathNames, (uint8)EMediaPlateEventState::Reverse);
-		}),
-		FCanExecuteAction::CreateLambda([this]{
-			TObjectPtr<UMediaPlayer> MediaPlayer = MediaPlate->GetMediaPlayer();
-			return (MediaPlayer != nullptr) && MediaPlayer->IsReady() && MediaPlayer->SupportsRate(GetReverseRate(), false);
-		})
+		FExecuteAction::CreateSP( this, &FMediaPlateEditorToolkit::HandleMediaPlateEvent, EMediaPlateEventState::Reverse ),
+		FCanExecuteAction::CreateSP( this, &FMediaPlateEditorToolkit::CanHandleMediaPlateEvent, EMediaPlateEventState::Reverse )
 	);
 
 	ToolkitCommands->MapAction(
 		Commands.RewindMedia,
-		FExecuteAction::CreateLambda([this, MediaAssets] {
-			MediaPlate->GetMediaPlayer()->Rewind();
-
-			TArray<FString> ActorsPathNames = { MediaPlate->GetOwner()->GetPathName() };
-			MediaAssets->BroadcastOnMediaStateChangedEvent(ActorsPathNames, (uint8)EMediaPlateEventState::Rewind);
-		}),
-		FCanExecuteAction::CreateLambda([this]{
-			TObjectPtr<UMediaPlayer> MediaPlayer = MediaPlate->GetMediaPlayer(); 
-			return (MediaPlayer != nullptr) && MediaPlate->GetMediaPlayer()->IsReady() && MediaPlayer->SupportsSeeking() && MediaPlayer->GetTime() > FTimespan::Zero();
-		})
+		FExecuteAction::CreateSP( this, &FMediaPlateEditorToolkit::HandleMediaPlateEvent, EMediaPlateEventState::Rewind ),
+		FCanExecuteAction::CreateSP( this, &FMediaPlateEditorToolkit::CanHandleMediaPlateEvent, EMediaPlateEventState::Rewind )
 	);
 }
 
@@ -466,6 +386,21 @@ float FMediaPlateEditorToolkit::GetReverseRate() const
 /* FMediaPlayerEditorToolkit callbacks
  *****************************************************************************/
 
+bool FMediaPlateEditorToolkit::CanHandleMediaPlateEvent(EMediaPlateEventState InState)
+{
+	if (MediaPlate)
+	{
+		return FMediaPlateCustomization::IsMediaPlateEventAllowedForPlayer(InState, MediaPlate->GetMediaPlayer());
+	}
+	return false;
+}
+
+void FMediaPlateEditorToolkit::HandleMediaPlateEvent(EMediaPlateEventState InState)
+{
+	const TArray<TWeakObjectPtr<UMediaPlateComponent>, TInlineAllocator<1>> MediaPlates = {MediaPlate};
+	FMediaPlateCustomization::HandleMediaPlateEvent(MediaPlates, InState);
+}
+
 TSharedRef<SDockTab> FMediaPlateEditorToolkit::HandleTabManagerSpawnTab(const FSpawnTabArgs& Args, FName TabIdentifier)
 {
 	UMediaPlayer* MediaPlayer = MediaPlate->GetMediaPlayer();
diff --git a/Engine/Plugins/Media/MediaPlate/Source/MediaPlateEditor/Private/Toolkits/MediaPlateEditorToolkit.h b/Engine/Plugins/Media/MediaPlate/Source/MediaPlateEditor/Private/Toolkits/MediaPlateEditorToolkit.h
--- a/Engine/Plugins/Media/MediaPlate/Source/MediaPlateEditor/Private/Toolkits/MediaPlateEditorToolkit.h
+++ b/Engine/Plugins/Media/MediaPlate/Source/MediaPlateEditor/Private/Toolkits/MediaPlateEditorToolkit.h
@@ -9,6 +9,7 @@
 #include "Toolkits/AssetEditorToolkit.h"
 #include "UObject/GCObject.h"
 
+enum class EMediaPlateEventState : uint8;
 class UMediaPlateComponent;
 
 /**
@@ -84,6 +85,11 @@ protected:
 	float GetReverseRate() const;
 
 private:
+	/** Callback to validate playback button events. */
+	bool CanHandleMediaPlateEvent(EMediaPlateEventState InState);
+
+	/** Callback to handle playback button events. */
+	void HandleMediaPlateEvent(EMediaPlateEventState InState);
 
 	/** Callback for spawning tabs. */
 	TSharedRef<SDockTab> HandleTabManagerSpawnTab(const FSpawnTabArgs& Args, FName TabIdentifier);
