[MediaPlateEditor] Fix potential null dereference.

diff --git a/Engine/Plugins/Media/MediaPlate/Source/MediaPlateEditor/Private/MediaPlateCustomization.cpp b/Engine/Plugins/Media/MediaPlate/Source/MediaPlateEditor/Private/MediaPlateCustomization.cpp
--- a/Engine/Plugins/Media/MediaPlate/Source/MediaPlateEditor/Private/MediaPlateCustomization.cpp
+++ b/Engine/Plugins/Media/MediaPlate/Source/MediaPlateEditor/Private/MediaPlateCustomization.cpp
@@ -1112,7 +1112,6 @@ void FMediaPlateCustomization::OnButtonEvent(EMediaPlateEventState InState)
 
 void FMediaPlateCustomization::HandleMediaPlateEvent(TConstArrayView<TWeakObjectPtr<UMediaPlateComponent>> InMediaPlates, EMediaPlateEventState InState)
 {
-	IMediaAssetsModule* MediaAssets = FModuleManager::LoadModulePtr<IMediaAssetsModule>("MediaAssets");
 	TArray<FString> ActorsPathNames;
 	ActorsPathNames.Reserve(InMediaPlates.Num());
 	for (const TWeakObjectPtr<UMediaPlateComponent>& MediaPlateWeak : InMediaPlates)
@@ -1130,8 +1129,7 @@ void FMediaPlateCustomization::HandleMediaPlateEvent(TConstArrayView<TWeakObject
 		if (InState == EMediaPlateEventState::Open)
 		{
 			// Tell the editor module that this media plate is playing.
-			FMediaPlateEditorModule* EditorModule = FModuleManager::LoadModulePtr<FMediaPlateEditorModule>("MediaPlateEditor");
-			if (EditorModule != nullptr)
+			if (FMediaPlateEditorModule* EditorModule = FModuleManager::LoadModulePtr<FMediaPlateEditorModule>("MediaPlateEditor"))
 			{
 				EditorModule->MediaPlateStartedPlayback(MediaPlate);
 			}
@@ -1139,7 +1137,11 @@ void FMediaPlateCustomization::HandleMediaPlateEvent(TConstArrayView<TWeakObject
 
 		MediaPlate->SwitchStates(InState);
 	}
-	MediaAssets->BroadcastOnMediaStateChangedEvent(ActorsPathNames, (uint8)InState);
+
+	if (IMediaAssetsModule* MediaAssets = FModuleManager::LoadModulePtr<IMediaAssetsModule>("MediaAssets"))
+	{
+		MediaAssets->BroadcastOnMediaStateChangedEvent(ActorsPathNames, (uint8)InState);
+	}
 }
 
 FReply FMediaPlateCustomization::OnOpenMediaPlate()
