[MediaPlateEditor] Fixing Next and Previous Buttons

Regression introduced in CL 46234211.

diff --git a/Engine/Plugins/Media/MediaPlate/Source/MediaPlate/Private/MediaPlateComponent.cpp b/Engine/Plugins/Media/MediaPlate/Source/MediaPlate/Private/MediaPlateComponent.cpp
--- a/Engine/Plugins/Media/MediaPlate/Source/MediaPlate/Private/MediaPlateComponent.cpp
+++ b/Engine/Plugins/Media/MediaPlate/Source/MediaPlate/Private/MediaPlateComponent.cpp
@@ -1884,18 +1884,12 @@ void UMediaPlateComponent::SwitchStates(EMediaPlateEventState State)
 		break;
 	case EMediaPlateEventState::Next:
 		{
-			if (MediaPlayer != nullptr)
-			{
-				MediaPlayer->Next();
-			}
+			Next();
 		}
 		break;
 	case EMediaPlateEventState::Previous:
 		{
-			if (MediaPlayer != nullptr)
-			{
-				MediaPlayer->Previous();
-			}
+			Previous();
 		}
 		break;
 	case EMediaPlateEventState::MAX:
diff --git a/Engine/Plugins/Media/MediaPlate/Source/MediaPlateEditor/Private/Toolkits/MediaPlateEditorToolkit.cpp b/Engine/Plugins/Media/MediaPlate/Source/MediaPlateEditor/Private/Toolkits/MediaPlateEditorToolkit.cpp
--- a/Engine/Plugins/Media/MediaPlate/Source/MediaPlateEditor/Private/Toolkits/MediaPlateEditorToolkit.cpp
+++ b/Engine/Plugins/Media/MediaPlate/Source/MediaPlateEditor/Private/Toolkits/MediaPlateEditorToolkit.cpp
@@ -390,7 +390,18 @@ bool FMediaPlateEditorToolkit::CanHandleMediaPlateEvent(EMediaPlateEventState In
 {
 	if (MediaPlate)
 	{
-		return FMediaPlateCustomization::IsMediaPlateEventAllowedForPlayer(InState, MediaPlate->GetMediaPlayer());
+		switch (InState)
+		{
+			case EMediaPlateEventState::Next:
+				return (MediaPlate->GetMediaPlaylist() != nullptr)
+					&& (MediaPlate->GetMediaPlaylist()->Num() > 1);
+			case EMediaPlateEventState::Previous:
+				return (MediaPlate->GetMediaPlaylist() != nullptr)
+					&& (MediaPlate->GetMediaPlaylist()->Num() > 1)
+					&& (MediaPlate->PlaylistIndex > 0); // Match UMediaPlateComponent::Previous() conditions.
+		default:
+			return FMediaPlateCustomization::IsMediaPlateEventAllowedForPlayer(InState, MediaPlate->GetMediaPlayer());	
+		}
 	}
 	return false;
 }
