MediaPlate: Light touch fix for Hidden In Game toggle causing the plate to lose its material in multi-user.

#rnx

diff --git a/Engine/Plugins/Media/MediaPlate/Source/MediaPlate/Private/MediaPlate.cpp b/Engine/Plugins/Media/MediaPlate/Source/MediaPlate/Private/MediaPlate.cpp
--- a/Engine/Plugins/Media/MediaPlate/Source/MediaPlate/Private/MediaPlate.cpp
+++ b/Engine/Plugins/Media/MediaPlate/Source/MediaPlate/Private/MediaPlate.cpp
@@ -157,6 +157,12 @@ void AMediaPlate::PostRegisterAllComponents()
 	if (IsValid(StaticMeshComponent))
 	{
 		ConditionallyEnableHoldoutComposite();
+
+#if WITH_EDITOR
+		// Update cached "last" members to prevent unnecessary material updates on static mesh edits (see ApplyCurrentMaterial).
+		LastMaterial = StaticMeshComponent->GetMaterial(0);
+		LastOverlayMaterial = StaticMeshComponent->GetOverlayMaterial();
+#endif
 	}
 }
 
