[MediaPlateEditor] Refactor: Unifying media details panel with the media source editor.

Making the media source detail panel reusable to remove code duplication.

#Synced-CL 48520225

diff --git a/Engine/Plugins/Media/MediaPlate/Source/MediaPlateEditor/Private/Widgets/SMediaPlateEditorMediaDetails.cpp b/Engine/Plugins/Media/MediaPlate/Source/MediaPlateEditor/Private/Widgets/SMediaPlateEditorMediaDetails.cpp
--- a/Engine/Plugins/Media/MediaPlate/Source/MediaPlateEditor/Private/Widgets/SMediaPlateEditorMediaDetails.cpp
+++ b/Engine/Plugins/Media/MediaPlate/Source/MediaPlateEditor/Private/Widgets/SMediaPlateEditorMediaDetails.cpp
@@ -1,14 +1,11 @@
 // Copyright Epic Games, Inc. All Rights Reserved.
 
 #include "Widgets/SMediaPlateEditorMediaDetails.h"
-#include "IDetailsView.h"
 #include "MediaPlateComponent.h"
 #include "MediaPlayer.h"
 #include "MediaTexture.h"
 #include "Modules/ModuleManager.h"
-#include "PropertyEditorModule.h"
 #include "Widgets/Layout/SScrollBox.h"
-#include "Widgets/Text/STextBlock.h"
 
 #define LOCTEXT_NAMESPACE "SMediaPlateEditorMediaDetails"
 
@@ -20,242 +17,25 @@ void SMediaPlateEditorMediaDetails::Construct(const FArguments& InArgs,
 {
 	MediaPlate = &InMediaPlate;
 
-	const ISlateStyle& SlateStyle = FAppStyle::Get();
-	const FName StyleName = TEXT("SmallText");
-
-	ChildSlot
-		[
-			SNew(SScrollBox)
-
-			// Add details.
-			+ SScrollBox::Slot()
-				[
-					SNew(SHorizontalBox)
-
-					// Left side.
-					+ SHorizontalBox::Slot()
-						[
-							SNew(SVerticalBox)
-
-							// Player name.
-							+ SVerticalBox::Slot()
-								.AutoHeight()
-								.VAlign(VAlign_Center)
-								.Padding(4.0f)
-								[
-									SAssignNew(MediaPlayerName, STextBlock)
-									.TextStyle(SlateStyle, StyleName)
-								]
-
-							// Resolution.
-							+ SVerticalBox::Slot()
-								.AutoHeight()
-								.VAlign(VAlign_Center)
-								.Padding(4.0f)
-								[
-									SAssignNew(ResolutionText, STextBlock)
-									.TextStyle(SlateStyle, StyleName)
-								]
-
-							// Frame rate.
-							+ SVerticalBox::Slot()
-								.AutoHeight()
-								.VAlign(VAlign_Center)
-								.Padding(4.0f)
-								[
-									SAssignNew(FrameRateText, STextBlock)
-									.TextStyle(SlateStyle, StyleName)
-								]
-
-							// Resource size.
-							+ SVerticalBox::Slot()
-								.AutoHeight()
-								.VAlign(VAlign_Center)
-								.Padding(4.0f)
-								[
-									SAssignNew(ResourceSizeText, STextBlock)
-									.TextStyle(SlateStyle, StyleName)
-								]
-
-							// Method.
-							+ SVerticalBox::Slot()
-								.AutoHeight()
-								.VAlign(VAlign_Center)
-								.Padding(4.0f)
-								[
-									SAssignNew(MethodText, STextBlock)
-									.TextStyle(SlateStyle, StyleName)
-								]
-						]
-
-					// Right side.
-					+ SHorizontalBox::Slot()
-						[
-							SNew(SVerticalBox)
-
-							// Format.
-							+ SVerticalBox::Slot()
-								.AutoHeight()
-								.VAlign(VAlign_Center)
-								.Padding(4.0f)
-								[
-									SAssignNew(FormatText, STextBlock)
-									.TextStyle(SlateStyle, StyleName)
-								]
-
-							// LOD bias.
-							+ SVerticalBox::Slot()
-								.AutoHeight()
-								.VAlign(VAlign_Center)
-								.Padding(4.0f)
-								[
-									SAssignNew(LODBiasText, STextBlock)
-									.TextStyle(SlateStyle, StyleName)
-								]
-
-							// Num mips.
-							+ SVerticalBox::Slot()
-								.AutoHeight()
-								.VAlign(VAlign_Center)
-								.Padding(4.0f)
-								[
-									SAssignNew(NumMipsText, STextBlock)
-									.TextStyle(SlateStyle, StyleName)
-								]
-
-							// Num tiles.
-							+ SVerticalBox::Slot()
-								.AutoHeight()
-								.VAlign(VAlign_Center)
-								.Padding(4.0f)
-								[
-									SAssignNew(NumTilesText, STextBlock)
-									.TextStyle(SlateStyle, StyleName)
-								]
-							// Start Timecode.
-							+ SVerticalBox::Slot()
-								.AutoHeight()
-								.VAlign(VAlign_Center)
-								.Padding(4.0f)
-								[
-									SAssignNew(StartTimecodeText, STextBlock)
-									.TextStyle(SlateStyle, StyleName)
-								]
-							// Seek Performance
-							+ SVerticalBox::Slot()
-								.AutoHeight()
-								.VAlign(VAlign_Center)
-								.Padding(4.0f)
-								[
-									SAssignNew(SeekPerformance, STextBlock)
-									.TextStyle(SlateStyle, StyleName)
-								]
-						]
-				]
-
-		];
-
-	UpdateDetails();
+	SMediaPlayerEditorMediaDetails::Construct(SMediaPlayerEditorMediaDetails::FArguments(), nullptr, nullptr, TEXT("SmallText"));
 }
 
-
-void SMediaPlateEditorMediaDetails::Tick(const FGeometry& AllottedGeometry, const double InCurrentTime, const float InDeltaTime)
-{
-	// Call parent.
-	SCompoundWidget::Tick(AllottedGeometry, InCurrentTime, InDeltaTime);
-
-	UpdateDetails();
-}
-
-void SMediaPlateEditorMediaDetails::UpdateDetails()
+UMediaPlayer* SMediaPlateEditorMediaDetails::GetMediaPlayer() const
 {
-	FName PlayerName;
-	FString Format;
-	float FrameRate = 0.0f;
-	int32 LODBias = 0;
-	FText Method;
-	int32 NumMips = 0;
-	int32 NumTotalTiles = 0;
-	int64 ResourceSize = 0;
-	int32 SurfaceWidth = 0;
-	int32 SurfaceHeight = 0;
-	FString StartTimecode;
-	int32 KeyframeInterval = -1;
-
-	if (MediaPlate.IsValid())
+	if (UMediaPlateComponent* MediaPlateComponent = MediaPlate.Get())
 	{
-		// Get player info.
-		UMediaPlayer* MediaPlayer = MediaPlate->GetMediaPlayer();
-		if (MediaPlayer != nullptr)
-		{
-			PlayerName = MediaPlayer->GetPlayerName();
-			FrameRate = MediaPlayer->GetVideoTrackFrameRate(INDEX_NONE, INDEX_NONE);
-			Format = MediaPlayer->GetVideoTrackType(INDEX_NONE, INDEX_NONE);
-			FIntPoint NumTiles(EForceInit::ForceInitToZero);
-			MediaPlayer->GetMediaInfo<FIntPoint>(NumTiles, UMediaPlayer::MediaInfoNameSourceNumTiles.Resolve());
-			NumTotalTiles = NumTiles.X * NumTiles.Y;
-			MediaPlayer->GetMediaInfo<FString>(StartTimecode, UMediaPlayer::MediaInfoNameStartTimecodeValue.Resolve());
-			MediaPlayer->GetMediaInfo<int32>(KeyframeInterval, UMediaPlayer::MediaInfoNameKeyframeInterval.Resolve());
-		}
-
-		// Get texture info,
-		UMediaTexture* MediaTexture = MediaPlate->GetMediaTexture();
-		if (MediaTexture != nullptr)
-		{
-			LODBias = MediaTexture->GetCachedLODBias();
-			Method = MediaTexture->IsCurrentlyVirtualTextured() ?
-				LOCTEXT("MethodVirtualStreamed", "Virtual Streamed")
-				: (!MediaTexture->IsStreamable() ? LOCTEXT("QuickInfo_MethodNotStreamed", "Not Streamed")
-					: LOCTEXT("MethodStreamed", "Streamed"));
-			NumMips = MediaTexture->GetTextureNumMips();
-			ResourceSize = (MediaTexture->GetResourceSizeBytes(EResourceSizeMode::Exclusive) + 512) / 1024;
-			SurfaceWidth = static_cast<int32>(MediaTexture->GetSurfaceWidth());
-			SurfaceHeight = static_cast<int32>(MediaTexture->GetSurfaceHeight());
-		}
+		return MediaPlateComponent->GetMediaPlayer();
 	}
+	return nullptr;
+}
 
-	// Update text.
-	MediaPlayerName->SetText(FText::Format(LOCTEXT("Player", "Player: {0}"),
-		FText::FromName(PlayerName)));
-	FormatText->SetText(FText::Format(LOCTEXT("Format", "Format: {0}"),
-		FText::FromString(Format)));
-	FrameRateText->SetText(FText::Format(LOCTEXT("FrameRate", "Frame Rate: {0}"), 
-		FText::AsNumber(FrameRate)));
-	LODBiasText->SetText(FText::Format(LOCTEXT("LODBias", "Combined LOD Bias: {0}"),
-		FText::AsNumber(LODBias)));
-	MethodText->SetText(FText::Format(LOCTEXT("Method", "Method: {0}"), Method));
-	NumMipsText->SetText(FText::Format(LOCTEXT("NumberOfMips", "Mips: {0}"),
-		FText::AsNumber(NumMips)));
-	NumTilesText->SetText(FText::Format(LOCTEXT("NumberOfTiles", "Tiles: {0}"),
-		FText::AsNumber(NumTotalTiles)));
-	ResolutionText->SetText(FText::Format(LOCTEXT("Resolution", "Resolution: {0}x{1}"),
-		FText::AsNumber(SurfaceWidth), FText::AsNumber(SurfaceHeight)));
-	ResourceSizeText->SetText(FText::Format(LOCTEXT("ResourceSize", "Resource Size: {0} KB"),
-		FText::AsNumber(ResourceSize)));
-
-	if (!StartTimecode.IsEmpty())
-	{
-		StartTimecodeText->SetText(FText::Format(LOCTEXT("StartTimecode", "Start Timecode: {0}"), FText::FromString(StartTimecode)));
-	}
-	StartTimecodeText->SetVisibility(!StartTimecode.IsEmpty() ? EVisibility::Visible : EVisibility::Collapsed);
-	
-	/*
-		The keyframe interval is one of:
-		  -1 : no information returned (unknown)
-		   0 : unknown keyframe spacing, not every frame is a keyframe but the spacing is variable or cannot be determined
-		   1 : every frame is a keyframe
-		  >1 : every n'th frame is a keyframe
-	*/
-	if (KeyframeInterval >= 0)
-	{
-		SeekPerformance->SetText(FText::Format(LOCTEXT("SeekPerformance", "Seek Performance: {0}"), KeyframeInterval == 1 ? LOCTEXT("SeekPerformanceF", "Fast") : LOCTEXT("SeekPerformanceS", "Slow (GOP codec)")));
-		SeekPerformance->SetVisibility(EVisibility::Visible);
-	}
-	else
+UMediaTexture* SMediaPlateEditorMediaDetails::GetMediaTexture() const
+{
+	if (UMediaPlateComponent* MediaPlateComponent = MediaPlate.Get())
 	{
-		SeekPerformance->SetVisibility(EVisibility::Collapsed);
+		return MediaPlateComponent->GetMediaTexture();
 	}
-
+	return nullptr;
 }
 
 #undef LOCTEXT_NAMESPACE
diff --git a/Engine/Plugins/Media/MediaPlate/Source/MediaPlateEditor/Private/Widgets/SMediaPlateEditorMediaDetails.h b/Engine/Plugins/Media/MediaPlate/Source/MediaPlateEditor/Private/Widgets/SMediaPlateEditorMediaDetails.h
--- a/Engine/Plugins/Media/MediaPlate/Source/MediaPlateEditor/Private/Widgets/SMediaPlateEditorMediaDetails.h
+++ b/Engine/Plugins/Media/MediaPlate/Source/MediaPlateEditor/Private/Widgets/SMediaPlateEditorMediaDetails.h
@@ -5,7 +5,7 @@
 #include "CoreMinimal.h"
 #include "Styling/ISlateStyle.h"
 #include "Widgets/DeclarativeSyntaxSupport.h"
-#include "Widgets/SCompoundWidget.h"
+#include "Widgets/SMediaPlayerEditorMediaDetails.h"
 
 class STextBlock;
 class UMediaPlateComponent;
@@ -14,7 +14,7 @@ class UMediaPlateComponent;
  * Implements the details panel of the MediaPlate asset editor.
  */
 class SMediaPlateEditorMediaDetails
-	: public SCompoundWidget
+	: public SMediaPlayerEditorMediaDetails
 {
 public:
 
@@ -31,29 +31,13 @@ public:
 	 */
 	void Construct(const FArguments& InArgs, UMediaPlateComponent& InMediaPlate);
 
-	//~ SWidget interface
-	virtual void Tick(const FGeometry& AllottedGeometry, const double InCurrentTime, const float InDeltaTime) override;
+protected:
+	//~ Begin SMediaPlayerEditorMediaDetails
+	virtual UMediaPlayer* GetMediaPlayer() const override;
+	virtual UMediaTexture* GetMediaTexture() const override;
+	//~ End SMediaPlayerEditorMediaDetails
 
 private:
-
-	/**
-	 * Updates our widgets to reflect the current state.
-	 */
-	void UpdateDetails();
-
 	/** Pointer to the MediaPlate that is being viewed. */
 	TWeakObjectPtr<UMediaPlateComponent> MediaPlate;
-
-	/** Our widgets. */
-	TSharedPtr<STextBlock> MediaPlayerName;
-	TSharedPtr<STextBlock> FormatText;
-	TSharedPtr<STextBlock> FrameRateText;
-	TSharedPtr<STextBlock> LODBiasText;
-	TSharedPtr<STextBlock> MethodText;
-	TSharedPtr<STextBlock> NumMipsText;
-	TSharedPtr<STextBlock> NumTilesText;
-	TSharedPtr<STextBlock> ResolutionText;
-	TSharedPtr<STextBlock> ResourceSizeText;
-	TSharedPtr<STextBlock> StartTimecodeText;
-	TSharedPtr<STextBlock> SeekPerformance;
 };
diff --git a/Engine/Plugins/Media/MediaPlayerEditor/Source/MediaPlayerEditor/Private/Widgets/SMediaPlayerEditorMediaDetails.cpp b/Engine/Plugins/Media/MediaPlayerEditor/Source/MediaPlayerEditor/Private/Widgets/SMediaPlayerEditorMediaDetails.cpp
--- a/Engine/Plugins/Media/MediaPlayerEditor/Source/MediaPlayerEditor/Private/Widgets/SMediaPlayerEditorMediaDetails.cpp
+++ b/Engine/Plugins/Media/MediaPlayerEditor/Source/MediaPlayerEditor/Private/Widgets/SMediaPlayerEditorMediaDetails.cpp
@@ -15,10 +15,14 @@
  *****************************************************************************/
 
 void SMediaPlayerEditorMediaDetails::Construct(const FArguments& InArgs,
-	UMediaPlayer* InMediaPlayer, UMediaTexture* InMediaTexture)
+	UMediaPlayer* InMediaPlayer, UMediaTexture* InMediaTexture, FName InTextStyleName)
 {
-	MediaPlayer = InMediaPlayer;
-	MediaTexture = InMediaTexture;
+	// Remark: these can be null.
+	MediaPlayerWeak = InMediaPlayer;
+	MediaTextureWeak = InMediaTexture;
+
+	const ISlateStyle& SlateStyle = FAppStyle::Get();
+	const FName StyleName = InTextStyleName.IsNone() ? TEXT("NormalText") : InTextStyleName;
 
 	ChildSlot
 		[
@@ -41,6 +45,7 @@ void SMediaPlayerEditorMediaDetails::Construct(const FArguments& InArgs,
 								.Padding(4.0f)
 								[
 									SAssignNew(MediaPlayerName, STextBlock)
+									.TextStyle(SlateStyle, StyleName)
 								]
 
 							// Resolution.
@@ -50,6 +55,7 @@ void SMediaPlayerEditorMediaDetails::Construct(const FArguments& InArgs,
 								.Padding(4.0f)
 								[
 									SAssignNew(ResolutionText, STextBlock)
+									.TextStyle(SlateStyle, StyleName)
 								]
 
 							// Frame rate.
@@ -59,6 +65,7 @@ void SMediaPlayerEditorMediaDetails::Construct(const FArguments& InArgs,
 								.Padding(4.0f)
 								[
 									SAssignNew(FrameRateText, STextBlock)
+									.TextStyle(SlateStyle, StyleName)
 								]
 
 							// Resource size.
@@ -68,6 +75,7 @@ void SMediaPlayerEditorMediaDetails::Construct(const FArguments& InArgs,
 								.Padding(4.0f)
 								[
 									SAssignNew(ResourceSizeText, STextBlock)
+									.TextStyle(SlateStyle, StyleName)
 								]
 
 							// Method.
@@ -77,6 +85,7 @@ void SMediaPlayerEditorMediaDetails::Construct(const FArguments& InArgs,
 								.Padding(4.0f)
 								[
 									SAssignNew(MethodText, STextBlock)
+									.TextStyle(SlateStyle, StyleName)
 								]
 						]
 
@@ -92,6 +101,7 @@ void SMediaPlayerEditorMediaDetails::Construct(const FArguments& InArgs,
 								.Padding(4.0f)
 								[
 									SAssignNew(LODBiasText, STextBlock)
+									.TextStyle(SlateStyle, StyleName)
 								]
 
 							// Num mips.
@@ -101,6 +111,7 @@ void SMediaPlayerEditorMediaDetails::Construct(const FArguments& InArgs,
 								.Padding(4.0f)
 								[
 									SAssignNew(NumMipsText, STextBlock)
+									.TextStyle(SlateStyle, StyleName)
 								]
 
 							// Num tiles.
@@ -110,6 +121,7 @@ void SMediaPlayerEditorMediaDetails::Construct(const FArguments& InArgs,
 								.Padding(4.0f)
 								[
 									SAssignNew(NumTilesText, STextBlock)
+									.TextStyle(SlateStyle, StyleName)
 								]
 
 							// Format.
@@ -119,6 +131,7 @@ void SMediaPlayerEditorMediaDetails::Construct(const FArguments& InArgs,
 								.Padding(4.0f)
 								[
 									SAssignNew(FormatText, STextBlock)
+									.TextStyle(SlateStyle, StyleName)
 								]
 
 							// Start Timecode.
@@ -128,6 +141,7 @@ void SMediaPlayerEditorMediaDetails::Construct(const FArguments& InArgs,
 								.Padding(4.0f)
 								[
 									SAssignNew(StartTimecodeText, STextBlock)
+									.TextStyle(SlateStyle, StyleName)
 								]
 
 							// Seek Performance
@@ -137,6 +151,7 @@ void SMediaPlayerEditorMediaDetails::Construct(const FArguments& InArgs,
 								.Padding(4.0f)
 								[
 									SAssignNew(SeekPerformance, STextBlock)
+									.TextStyle(SlateStyle, StyleName)
 								]
 
 							// Video decoder provider
@@ -146,6 +161,7 @@ void SMediaPlayerEditorMediaDetails::Construct(const FArguments& InArgs,
 								.Padding(4.0f)
 								[
 									SAssignNew(VideoDecoderProvider, STextBlock)
+									.TextStyle(SlateStyle, StyleName)
 								]
 
 							// Audio decoder provider
@@ -155,6 +171,7 @@ void SMediaPlayerEditorMediaDetails::Construct(const FArguments& InArgs,
 								.Padding(4.0f)
 								[
 									SAssignNew(AudioDecoderProvider, STextBlock)
+									.TextStyle(SlateStyle, StyleName)
 								]
 						]
 				]
@@ -173,6 +190,16 @@ void SMediaPlayerEditorMediaDetails::Tick(const FGeometry& AllottedGeometry, con
 	UpdateDetails();
 }
 
+UMediaPlayer* SMediaPlayerEditorMediaDetails::GetMediaPlayer() const
+{
+	return MediaPlayerWeak.Get();
+}
+
+UMediaTexture* SMediaPlayerEditorMediaDetails::GetMediaTexture() const
+{
+	return MediaTextureWeak.Get();
+}
+
 void SMediaPlayerEditorMediaDetails::UpdateDetails()
 {
 	FName PlayerName;
@@ -190,6 +217,9 @@ void SMediaPlayerEditorMediaDetails::UpdateDetails()
 	FString VideoDecoderImpl;
 	FString AudioDecoderImpl;
 
+	const UMediaPlayer* MediaPlayer = GetMediaPlayer();
+	UMediaTexture* MediaTexture = GetMediaTexture();
+	
 	// Get player info.
 	if (MediaPlayer != nullptr)
 	{
diff --git a/Engine/Plugins/Media/MediaPlayerEditor/Source/MediaPlayerEditor/Public/Widgets/SMediaPlayerEditorMediaDetails.h b/Engine/Plugins/Media/MediaPlayerEditor/Source/MediaPlayerEditor/Public/Widgets/SMediaPlayerEditorMediaDetails.h
--- a/Engine/Plugins/Media/MediaPlayerEditor/Source/MediaPlayerEditor/Public/Widgets/SMediaPlayerEditorMediaDetails.h
+++ b/Engine/Plugins/Media/MediaPlayerEditor/Source/MediaPlayerEditor/Public/Widgets/SMediaPlayerEditorMediaDetails.h
@@ -6,6 +6,8 @@
 #include "Widgets/DeclarativeSyntaxSupport.h"
 #include "Widgets/SCompoundWidget.h"
 
+#define UE_API MEDIAPLAYEREDITOR_API
+
 class STextBlock;
 class UMediaPlayer;
 class UMediaTexture;;
@@ -29,12 +31,20 @@ public:
 	 * @param InArgs			The declaration data for this widget.
 	 * @param InMediaPlayer		The MediaPlayer to show the details for.
 	 * @param InMediaTexture	The MediaTexture to show the details for.
+	 * @param InTextStyleName	The text style name to use.
 	 *
 	 */
-	void Construct(const FArguments& InArgs, UMediaPlayer* InMediaPlate, UMediaTexture* InMediaTexture);
+	UE_API void Construct(const FArguments& InArgs, UMediaPlayer* InMediaPlayer, UMediaTexture* InMediaTexture, FName InTextStyleName = NAME_None);
 
 	//~ SWidget interface
-	virtual void Tick(const FGeometry& AllottedGeometry, const double InCurrentTime, const float InDeltaTime) override;
+	UE_API virtual void Tick(const FGeometry& AllottedGeometry, const double InCurrentTime, const float InDeltaTime) override;
+
+protected:
+	/** Accessor for media player. */
+	UE_API virtual UMediaPlayer* GetMediaPlayer() const;
+
+	/** Accessor for the media texture. */
+	UE_API virtual UMediaTexture* GetMediaTexture() const;
 
 private:
 
@@ -44,9 +54,9 @@ private:
 	void UpdateDetails();
 
 	/** Pointer to the MediaPlayer that is being viewed. */
-	TWeakObjectPtr<UMediaPlayer> MediaPlayer;
+	TWeakObjectPtr<UMediaPlayer> MediaPlayerWeak;
 	/** Pointer to the MediaTexture that is being viewed. */
-	TWeakObjectPtr<UMediaTexture> MediaTexture;
+	TWeakObjectPtr<UMediaTexture> MediaTextureWeak;
 
 	/** Our widgets. */
 	TSharedPtr<STextBlock> MediaPlayerName;
@@ -63,3 +73,5 @@ private:
 	TSharedPtr<STextBlock> VideoDecoderProvider;
 	TSharedPtr<STextBlock> AudioDecoderProvider;
 };
+
+#undef UE_API
\ No newline at end of file
