[MediaPlateEditor] The user is randomly unable to create a media playlist or media source from the media plate resource asset picker menu.

Explicitly setting AllowCreate(true) to make sure it is enabled.

#Synced-CL 49341242

diff --git a/Engine/Plugins/Media/MediaPlate/Source/MediaPlateEditor/Private/MediaPlateResourceCustomization.cpp b/Engine/Plugins/Media/MediaPlate/Source/MediaPlateEditor/Private/MediaPlateResourceCustomization.cpp
--- a/Engine/Plugins/Media/MediaPlate/Source/MediaPlateEditor/Private/MediaPlateResourceCustomization.cpp
+++ b/Engine/Plugins/Media/MediaPlate/Source/MediaPlateEditor/Private/MediaPlateResourceCustomization.cpp
@@ -76,6 +76,7 @@ void FMediaPlateResourceCustomization::CustomizeHeader(TSharedRef<IPropertyHandl
 				[
 					SNew(SObjectPropertyEntryBox)
 					.AllowedClass(UMediaSource::StaticClass())
+					.AllowCreate(true)
 					.PropertyHandle(MediaAssetPropertyHandle)
 				]
 			]
@@ -104,6 +105,7 @@ void FMediaPlateResourceCustomization::CustomizeHeader(TSharedRef<IPropertyHandl
 				[
 					SNew(SObjectPropertyEntryBox)
 					.AllowedClass(UMediaPlaylist::StaticClass())
+					.AllowCreate(true)
 					.PropertyHandle(SourcePlaylistPropertyHandle)
 				]
 			]
