[MediaPlate] Support multiple media textures in the "material instance constant" code path.

- Texture parameter code is moved from SetMIDParameters into private util functions to be reused in both the MID and MIC code paths.
- Adding multiple texture parameter setup in CreateMaterialInstanceConstant (main goal of this change).

Cross fade: For dynamic blend weights to work in editor, user still needs to "create dynamic material" manually, but it will now have the correct media textures in the parameters.

#Synced-CL 52257653

diff --git a/Engine/Plugins/Media/MediaPlate/Source/MediaPlate/Private/MediaPlate.cpp b/Engine/Plugins/Media/MediaPlate/Source/MediaPlate/Private/MediaPlate.cpp
--- a/Engine/Plugins/Media/MediaPlate/Source/MediaPlate/Private/MediaPlate.cpp
+++ b/Engine/Plugins/Media/MediaPlate/Source/MediaPlate/Private/MediaPlate.cpp
@@ -45,6 +45,35 @@ namespace UE::MediaPlate::Private
 
 		return false;
 	}
+
+	/** Returns the number of texture parameters that begin with the "MediaTexture" name string in the given material instance. */
+	int32 GetNumMediaTextureParameters(const UMaterialInstance* InMaterialInstance, const FString& InMediaTextureString)
+	{
+		int32 NumTextures = 0;
+		for (const FTextureParameterValue& Param : InMaterialInstance->TextureParameterValues)
+		{
+			const FString Name = Param.ParameterInfo.Name.ToString();
+			if (Name.StartsWith(InMediaTextureString))
+			{
+				++NumTextures;
+			}
+		}
+		return NumTextures;
+	}
+
+	/** Returns the expected material texture parameter name for the given texture index. */
+	FName GetTextureParameterName(const FString& InMediaTextureString, int32 InTextureIndex)
+	{
+		// No index is appended for index zero.
+		if (InTextureIndex == 0)
+		{
+			return FName(InMediaTextureString);
+		}
+
+		FString TextureParameterName = InMediaTextureString;
+		TextureParameterName.AppendInt(InTextureIndex);
+		return FName(TextureParameterName);
+	}
 }
 
 AMediaPlate::AMediaPlate(const FObjectInitializer& ObjectInitializer)
@@ -290,9 +319,32 @@ UMaterialInterface* AMediaPlate::CreateMaterialInstanceConstant(UMaterialInterfa
 		NewObject<UMaterialInstanceConstant>(StaticMeshComponent, MaterialUniqueName, RF_Transactional);
 	MaterialInstance->SetParentEditorOnly(InMaterial);
 	MaterialInstance->CopyMaterialUniformParametersEditorOnly(InMaterial);
+
+	using namespace UE::MediaPlate::Private;
+	const FString MediaTextureNameString = MediaTextureName.Resolve().ToString();
+	const int32 NumTextures = GetNumMediaTextureParameters(MaterialInstance, MediaTextureNameString);
+
+	if (NumTextures == 0)
+	{
+		UE_LOG(LogMediaPlate, Warning, TEXT("Material \"%s\" doesn't have any texture parameters that start with \"%s\"."), *MaterialName, *MediaTextureNameString);
+	}
+	
+	MediaPlateComponent->SetNumberOfTextures(FMath::Max(NumTextures, 1));
+
 	MaterialInstance->SetTextureParameterValueEditorOnly(
 		FMaterialParameterInfo(MediaTextureName),
 		MediaPlateComponent->GetMediaTexture());
+
+	for (int32 TextureIndex = 1; TextureIndex < NumTextures; ++TextureIndex)
+	{
+		MaterialInstance->SetTextureParameterValueEditorOnly(
+			FMaterialParameterInfo(GetTextureParameterName(MediaTextureNameString, TextureIndex)),
+			MediaPlateComponent->GetMediaTexture(TextureIndex));
+	}
+
+	// Cross fade material: make sure the blend0 weight is set.
+	MaterialInstance->SetScalarParameterValueEditorOnly(FMaterialParameterInfo(TEXT("Blend0")), 1.0f);
+
 	MaterialInstance->PostEditChange();
 
 	// We force call post-load to indirectly call UpdateParameters() (for integration with VPUtilities plugin).
@@ -341,7 +393,6 @@ void AMediaPlate::ApplyMaterial(UMaterialInterface* Material)
 			}
 			else
 			{
-				MediaPlateComponent->SetNumberOfTextures(1);
 				Result = CreateMaterialInstanceConstant(Material);
 			}
 
@@ -391,7 +442,6 @@ void AMediaPlate::ApplyOverlayMaterial(UMaterialInterface* InOverlayMaterial)
 			}
 			else
 			{
-				MediaPlateComponent->SetNumberOfTextures(1);
 				Result = CreateMaterialInstanceConstant(InOverlayMaterial);
 			}
 
@@ -409,29 +459,22 @@ void AMediaPlate::ApplyOverlayMaterial(UMaterialInterface* InOverlayMaterial)
 
 void AMediaPlate::SetMIDParameters(UMaterialInstanceDynamic* InMaterial)
 {
-	InMaterial->SetTextureParameterValue(MediaTextureName, MediaPlateComponent->GetMediaTexture());
-	
-	int32 NumTextures = 0;
-	FString MediaTextureString = MediaTextureName.Resolve().ToString();
-	for (const struct FTextureParameterValue& Param : InMaterial->TextureParameterValues)
+	using namespace UE::MediaPlate::Private;
+	const FString MediaTextureString = MediaTextureName.Resolve().ToString();
+	const int32 NumTextures = GetNumMediaTextureParameters(InMaterial, MediaTextureString);
+
+	if (NumTextures == 0)
 	{
-		FString Name = Param.ParameterInfo.Name.ToString();
-		if (Name.StartsWith(MediaTextureString))
-		{
-			NumTextures++;
-		}
+		UE_LOG(LogMediaPlate, Warning, TEXT("Material \"%s\" doesn't have any texture parameters that start with \"%s\"."), *InMaterial->GetName(), *MediaTextureString);
 	}
-	MediaPlateComponent->SetNumberOfTextures(NumTextures);
+	
+	MediaPlateComponent->SetNumberOfTextures(FMath::Max(NumTextures, 1));
 
-	for (int32 Index = 0; Index < NumTextures; Index++)
+	InMaterial->SetTextureParameterValue(MediaTextureName, MediaPlateComponent->GetMediaTexture());
+	
+	for (int32 Index = 1; Index < NumTextures; Index++)
 	{
-		FString NameString = MediaTextureString;
-		if (Index != 0)
-		{
-			NameString.AppendInt(Index);
-		}
-		FName TextureParameterName = FName(*NameString);
-		InMaterial->SetTextureParameterValue(TextureParameterName,
+		InMaterial->SetTextureParameterValue(GetTextureParameterName(MediaTextureString, Index),
 			MediaPlateComponent->GetMediaTexture(Index));
 	}
 }
