[MediaPlateEditor] Prevent adding multiple media tracks under the same media plate binding.

Given that, by design in sequencer:
- the tracks under an object binding are not reorderable
- the sections cannot be moved or re-ordered accross tracks

It is preferable to prevent users from creating multiple tracks under a media plate binding to avoid the confusion and lead users to report bugs about the two problems above.
If a user wants to add another section for a new media in the track, the natural impulse is to click the top most level "+", which is on the binding, where it will find the choice to add a new "track", which is not the work flow we want the users to go for.

#Synced-CL 52257653

diff --git a/Engine/Plugins/Media/MediaPlate/Source/MediaPlateEditor/Private/Sequencer/MediaPlateTrackEditor.cpp b/Engine/Plugins/Media/MediaPlate/Source/MediaPlateEditor/Private/Sequencer/MediaPlateTrackEditor.cpp
--- a/Engine/Plugins/Media/MediaPlate/Source/MediaPlateEditor/Private/Sequencer/MediaPlateTrackEditor.cpp
+++ b/Engine/Plugins/Media/MediaPlate/Source/MediaPlateEditor/Private/Sequencer/MediaPlateTrackEditor.cpp
@@ -77,10 +77,11 @@ void FMediaPlateTrackEditor::BuildObjectBindingTrackMenu(FMenuBuilder& MenuBuild
 		{
 			MenuBuilder.AddMenuEntry(
 				LOCTEXT("AddTrack", "Media"),
-				LOCTEXT("AddAttachedTooltip", "Adds a media track attached to the object."),
+				LOCTEXT("AddAttachedTooltip", "Adds a media track attached to the object.\nOnly one track per binding is supported. Add Media Sections instead."),
 				FSlateIcon(),
 				FUIAction(
-					FExecuteAction::CreateRaw(this, &FMediaPlateTrackEditor::HandleAddMediaTrackToObjectBindingMenuEntryExecute, ObjectBindings)));
+					FExecuteAction::CreateRaw(this, &FMediaPlateTrackEditor::HandleAddMediaTrackToObjectBindingMenuEntryExecute, ObjectBindings),
+					FCanExecuteAction::CreateRaw(this, &FMediaPlateTrackEditor::HandleAddMediaTrackToObjectBindingMenuEntryCanExecute, ObjectBindings)));
 		}
 	}
 }
@@ -90,6 +91,34 @@ bool FMediaPlateTrackEditor::SupportsType(TSubclassOf<UMovieSceneTrack> TrackCla
 	return false;
 }
 
+bool FMediaPlateTrackEditor::HandleAddMediaTrackToObjectBindingMenuEntryCanExecute(TArray<FGuid> InObjectBindingID)
+{
+	const TSharedPtr<ISequencer> SequencerPtr = GetSequencer();
+	const UMovieSceneSequence* Sequence = SequencerPtr ? SequencerPtr->GetFocusedMovieSceneSequence() : nullptr; 
+	const UMovieScene* MovieScene = Sequence ? Sequence->GetMovieScene() : nullptr;
+
+	if (!MovieScene)
+	{
+		return false;
+	}
+	
+	// Only allow adding a media track under this object binding if it doesn't already have one.
+	for (const FGuid& ObjectBindingID : InObjectBindingID)
+	{
+		if (const FMovieSceneBinding* Binding = MovieScene->FindBinding(ObjectBindingID))
+		{
+			for (const UMovieSceneTrack* Track : Binding->GetTracks())
+			{
+				if (Track && Track->IsA<UMovieSceneMediaTrack>())
+				{
+					return false;
+				}
+			}
+		}
+	}
+	return true;
+}
+
 void FMediaPlateTrackEditor::HandleAddMediaTrackToObjectBindingMenuEntryExecute(TArray<FGuid> InObjectBindingIDs)
 {
 	UMovieScene* FocusedMovieScene = GetFocusedMovieScene();
diff --git a/Engine/Plugins/Media/MediaPlate/Source/MediaPlateEditor/Private/Sequencer/MediaPlateTrackEditor.h b/Engine/Plugins/Media/MediaPlate/Source/MediaPlateEditor/Private/Sequencer/MediaPlateTrackEditor.h
--- a/Engine/Plugins/Media/MediaPlate/Source/MediaPlateEditor/Private/Sequencer/MediaPlateTrackEditor.h
+++ b/Engine/Plugins/Media/MediaPlate/Source/MediaPlateEditor/Private/Sequencer/MediaPlateTrackEditor.h
@@ -60,6 +60,10 @@ public:
 	virtual void OnRelease() override;
 
 private:
+	/**
+	 * Callback for enabling the action to add a media track to an object binding.
+	 */
+	bool HandleAddMediaTrackToObjectBindingMenuEntryCanExecute(TArray<FGuid> InObjectBindingID);
 	/**
 	 * Callback for adding a media track to an object binding.
 	 */
